Change bootstrap program to wait for child process exit before returning. Fixes automated tests that require it for timing.

diff --git a/Engine/Source/Programs/Windows/BootstrapPackagedGame/Private/BootstrapPackagedGame.cpp b/Engine/Source/Programs/Windows/BootstrapPackagedGame/Private/BootstrapPackagedGame.cpp
--- a/Engine/Source/Programs/Windows/BootstrapPackagedGame/Private/BootstrapPackagedGame.cpp
+++ b/Engine/Source/Programs/Windows/BootstrapPackagedGame/Private/BootstrapPackagedGame.cpp
@@ -110,9 +110,13 @@ int SpawnTarget(WCHAR* CmdLine)
 		return 9005;
 	}
 
+	WaitForSingleObject(ProcessInfo.hProcess, INFINITE);
+	DWORD ExitCode = 9006;
+	GetExitCodeProcess(ProcessInfo.hProcess, &ExitCode);
+
 	CloseHandle(ProcessInfo.hThread);
 	CloseHandle(ProcessInfo.hProcess);
-	return 0;
+	return (int)ExitCode;
 }
 
 int WINAPI wWinMain(HINSTANCE hInstance, HINSTANCE hPrevInstance, TCHAR* CmdLine, int ShowCmd)
