Fixes for compile error introduced in VS2017 version 15.3. Can't include <windows.h> from the 8.1 SDK without prototyping IUnknown.

diff --git a/Engine/Source/Programs/Windows/BootstrapPackagedGame/Private/BootstrapPackagedGame.h b/Engine/Source/Programs/Windows/BootstrapPackagedGame/Private/BootstrapPackagedGame.h
--- a/Engine/Source/Programs/Windows/BootstrapPackagedGame/Private/BootstrapPackagedGame.h
+++ b/Engine/Source/Programs/Windows/BootstrapPackagedGame/Private/BootstrapPackagedGame.h
@@ -2,6 +2,8 @@
 
 #pragma once
 
+struct IUnknown;
+
 #include <windows.h>
 #include <tchar.h>
 #include <assert.h>
diff --git a/Engine/Source/Runtime/Core/Public/Windows/MinWindows.h b/Engine/Source/Runtime/Core/Public/Windows/MinWindows.h
--- a/Engine/Source/Runtime/Core/Public/Windows/MinWindows.h
+++ b/Engine/Source/Runtime/Core/Public/Windows/MinWindows.h
@@ -62,5 +62,7 @@
 #define NOPROXYSTUB
 #define NORPC
 
+struct IUnknown;
+
 // Finally now we can include windows.h
 #include <Windows.h>
