Fix warning compiling BootstrapPackagedGame Win32 Shipping, due to use of 64-bit pointers in the Windows 10 SDK that require 8-byte alignment (which is being overriden via /Zp4 from UBT on Win32).

diff --git a/Engine/Source/Programs/Windows/BootstrapPackagedGame/Private/BootstrapPackagedGame.h b/Engine/Source/Programs/Windows/BootstrapPackagedGame/Private/BootstrapPackagedGame.h
--- a/Engine/Source/Programs/Windows/BootstrapPackagedGame/Private/BootstrapPackagedGame.h
+++ b/Engine/Source/Programs/Windows/BootstrapPackagedGame/Private/BootstrapPackagedGame.h
@@ -4,7 +4,10 @@
 
 struct IUnknown;
 
+#pragma pack(push, 8)
 #include <windows.h>
+#pragma pack(pop)
+
 #include <tchar.h>
 #include <assert.h>
 #include <stdio.h>
