Fixing compile time parsing so it doesn't accidentally strip the first line.
Fixing sorting in CompileTimeAnalyzer.

#rnx

diff --git a/Engine/Extras/cl-filter/cl-filter.cpp b/Engine/Extras/cl-filter/cl-filter.cpp
--- a/Engine/Extras/cl-filter/cl-filter.cpp
+++ b/Engine/Extras/cl-filter/cl-filter.cpp
@@ -253,6 +253,11 @@ int wmain(int ArgC, const wchar_t* ArgV[])
 				//   handle by what our current state is.
 				const char StartTimingInfoText[] = "Include Headers:";
 				
+				if (!InTimingInfo && memcmp(Buffer + LineStart, StartTimingInfoText, 16) == 0)
+				{
+					InTimingInfo = true;	
+				}
+
 				if (InTimingInfo)
 				{
 					DWORD BytesWritten;
@@ -260,12 +265,6 @@ int wmain(int ArgC, const wchar_t* ArgV[])
 					LineStart = LineEnd;
 					continue;
 				}
-				else if (memcmp(Buffer + LineStart, StartTimingInfoText, 16) == 0)
-				{
-					InTimingInfo = true;
-					LineStart = LineEnd;
-					continue;
-				}
 			}
 
 			// If we didn't write anything out, write it to stdout
diff --git a/Engine/Source/Programs/Windows/CompileTimeAnalyzer/CompileTimeAnalyzer/Controls/TimingDataGrid.xaml.cs b/Engine/Source/Programs/Windows/CompileTimeAnalyzer/CompileTimeAnalyzer/Controls/TimingDataGrid.xaml.cs
--- a/Engine/Source/Programs/Windows/CompileTimeAnalyzer/CompileTimeAnalyzer/Controls/TimingDataGrid.xaml.cs
+++ b/Engine/Source/Programs/Windows/CompileTimeAnalyzer/CompileTimeAnalyzer/Controls/TimingDataGrid.xaml.cs
@@ -1,9 +1,14 @@
 // Copyright 1998-2019 Epic Games, Inc. All Rights Reserved.
 
+using System.Collections.Generic;
+using System.Collections.ObjectModel;
 using System.ComponentModel;
+using System.Linq;
+using System.Reflection;
 using System.Windows.Controls;
 using System.Windows.Data;
 using System.Windows.Input;
+using Timing_Data_Investigator.Models;
 
 namespace Timing_Data_Investigator.Controls
 {
@@ -19,7 +24,8 @@ namespace Timing_Data_Investigator.Controls
 
         private void Grid_Sorting(object sender, DataGridSortingEventArgs e)
         {
-			CollectionViewSource GridModel = DataContext as CollectionViewSource;
+			
+			
 			switch (e.Column.SortDirection)
 			{
 				case ListSortDirection.Ascending:
@@ -41,15 +47,35 @@ namespace Timing_Data_Investigator.Controls
 					}
 			}
 
-			GridModel?.SortDescriptions?.Clear();
-			if (e.Column.SortDirection != null)
-			{
-				GridModel.SortDescriptions.Add(new SortDescription(e.Column.SortMemberPath, e.Column.SortDirection.Value));
-			}
+			TreeGridFlatModel GridModel = (TreeGridFlatModel)DataContext;
+			PropertyInfo SortProperty = typeof(TimingDataViewModel).GetProperty(e.Column.SortMemberPath);
+			GridModel.Sort(SortProperty, e.Column.SortDirection);
 
             e.Handled = true;
         }
 
+		private void SortNodes(ObservableCollection<TreeGridElement> NewFlatModel, TreeGridElement Node, ListSortDirection SortDirection, PropertyInfo SortProperty)
+		{
+			NewFlatModel.Add(Node);
+			if (Node.IsExpanded && Node.HasChildren)
+			{
+				IOrderedEnumerable<TreeGridElement> SortedChildren;
+				if (SortDirection == ListSortDirection.Ascending)
+				{
+					SortedChildren = Node.Children.OrderBy(c => SortProperty.GetValue(c));
+				}
+				else
+				{
+					SortedChildren = Node.Children.OrderByDescending(c => SortProperty.GetValue(c));
+				}
+
+				foreach (TreeGridElement Child in SortedChildren)
+				{
+					SortNodes(NewFlatModel, Child, SortDirection, SortProperty);
+				}
+			}
+		}
+
         private void Grid_KeyDown(object sender, KeyEventArgs e)
         {
 			TreeGridElement SelectedData = Grid.SelectedItem as TreeGridElement;
diff --git a/Engine/Source/Programs/Windows/CompileTimeAnalyzer/CompileTimeAnalyzer/Controls/TimingDataTab.xaml.cs b/Engine/Source/Programs/Windows/CompileTimeAnalyzer/CompileTimeAnalyzer/Controls/TimingDataTab.xaml.cs
--- a/Engine/Source/Programs/Windows/CompileTimeAnalyzer/CompileTimeAnalyzer/Controls/TimingDataTab.xaml.cs
+++ b/Engine/Source/Programs/Windows/CompileTimeAnalyzer/CompileTimeAnalyzer/Controls/TimingDataTab.xaml.cs
@@ -4,6 +4,7 @@ using System.Collections.Generic;
 using System.Collections.ObjectModel;
 using System.ComponentModel;
 using System.Linq;
+using System.Reflection;
 using System.Text.RegularExpressions;
 using System.Windows;
 using System.Windows.Controls;
@@ -26,11 +27,6 @@ namespace Timing_Data_Investigator.Controls
         private TreeGridModel FunctionsModel;
         private TreeGridModel GroupedFunctionsModel;
 
-		private CollectionViewSource FilesViewSource = new CollectionViewSource();
-		private CollectionViewSource IncludesViewSource = new CollectionViewSource();
-		private CollectionViewSource ClassesViewSource = new CollectionViewSource();
-		private CollectionViewSource FunctionsViewSource = new CollectionViewSource();
-
 		private Dictionary<TimingDataViewModel, TabState> TabStates = new Dictionary<TimingDataViewModel, TabState>();
 
         private class TabState
@@ -55,10 +51,6 @@ namespace Timing_Data_Investigator.Controls
             InitializeComponent();
             DataContextChanged += TimingDataTab_DataContextChanged;
             FilesGrid.Grid.RowStyle = FindResource("FilesRowStyle") as Style;
-			FilesGrid.DataContext = FilesViewSource;
-			IncludesGrid.DataContext = IncludesViewSource;
-			ClassesGrid.DataContext = ClassesViewSource;
-			FunctionsGrid.DataContext = FunctionsViewSource;
 		}
 
         private void TimingDataTab_DataContextChanged(object sender, DependencyPropertyChangedEventArgs e)
@@ -127,27 +119,27 @@ namespace Timing_Data_Investigator.Controls
 				RefreshTabState(ClassesTab, State);
 				RefreshTabState(FunctionsTab, State);
 
-				SortModel(null, FilesGrid.Grid);
-				SortModel(null, IncludesGrid.Grid);
-				SortModel(null, ClassesGrid.Grid);
-				SortModel(null, FunctionsGrid.Grid);
+				SortModel(FilesGrid.Grid);
+				SortModel(IncludesGrid.Grid);
+				SortModel(ClassesGrid.Grid);
+				SortModel(FunctionsGrid.Grid);
 			}
         }
 
 		private void UpdateUIForAggregate(TimingDataViewModel TimingData)
 		{
 			FilesModel = GenerateTreeGridModel(TimingData.Children[0].Children);
-			IncludesModel = null;
-			ClassesModel = null;
-			FunctionsModel = null;
+			IncludesModel = GenerateTreeGridModel(TimingData.Children[1].Children);
+			ClassesModel = GenerateTreeGridModel(TimingData.Children[2].Children);
+			FunctionsModel = GenerateTreeGridModel(TimingData.Children[3].Children);
 			FlattenedIncludesModel = null;
 			GroupedClassesModel = null;
 			GroupedFunctionsModel = null;
 
 			UpdateTabVisibility(FilesModel.FlatModel, FilesTab, FilesGrid);
-			UpdateTabVisibility(TimingData.Children[1].Children, IncludesTab, IncludesGrid);
-			UpdateTabVisibility(TimingData.Children[2].Children, ClassesTab, ClassesGrid);
-			UpdateTabVisibility(TimingData.Children[3].Children, FunctionsTab, FunctionsGrid);
+			UpdateTabVisibility(IncludesModel.FlatModel, IncludesTab, IncludesGrid);
+			UpdateTabVisibility(ClassesModel.FlatModel, ClassesTab, ClassesGrid);
+			UpdateTabVisibility(FunctionsModel.FlatModel, FunctionsTab, FunctionsGrid);
 		}
 
 		private void UpdateUIForSingleFile(TimingDataViewModel TimingData)
@@ -176,8 +168,13 @@ namespace Timing_Data_Investigator.Controls
 			State.SortColumns[Tab].SortDirection = State.SortDirections[Tab];
 		}
 
-        private void SortModel(TreeGridModel Model, DataGrid Grid)
+        private void SortModel(DataGrid Grid)
         {
+			if (Grid.DataContext == null)
+			{
+				return;
+			}
+
             // Find the column that has a sorting value set, if any.
             DataGridColumn SortedColumn = Grid.Columns.FirstOrDefault(c => c.SortDirection != null);
             if (SortedColumn == null)
@@ -186,19 +183,15 @@ namespace Timing_Data_Investigator.Controls
                 return;
             }
 
-			CollectionViewSource CollectionToSort = (CollectionViewSource)Grid.DataContext;
-			CollectionToSort?.SortDescriptions?.Clear();
-			if (SortedColumn.SortDirection != null)
-			{
-				CollectionToSort.SortDescriptions.Add(new SortDescription(SortedColumn.SortMemberPath, SortedColumn.SortDirection.Value));
-			}
+			TreeGridFlatModel Model = (TreeGridFlatModel)Grid.DataContext;
+			PropertyInfo SortProperty = typeof(TimingDataViewModel).GetProperty(SortedColumn.SortMemberPath);
+			Model.Sort(SortProperty, SortedColumn.SortDirection);
 		}
 
-        private void UpdateTabVisibility(ObservableCollection<TreeGridElement> TabModel, TabItem Tab, TimingDataGrid Grid)
+        private void UpdateTabVisibility(TreeGridFlatModel TabModel, TabItem Tab, TimingDataGrid Grid)
         {
             Tab.Visibility = TabModel == null ? Visibility.Collapsed : Visibility.Visible;
-			CollectionViewSource Source = (CollectionViewSource)Grid.Grid.DataContext;
-			Source.Source = TabModel;
+			Grid.DataContext = TabModel;
         }
 
 		private TreeGridModel GenerateTreeGridModel(IEnumerable<TreeGridElement> Children)
@@ -369,14 +362,13 @@ namespace Timing_Data_Investigator.Controls
 
         private void UpdateGridModel(CheckBox CheckBox, TimingDataGrid Grid, TreeGridModel UncheckedModel, TreeGridModel CheckedModel)
         {
-			CollectionViewSource Source = (CollectionViewSource)Grid.DataContext;
 			if (CheckBox.IsChecked == true)
             {
-				Source.Source = CheckedModel?.FlatModel;
+				Grid.DataContext = CheckedModel?.FlatModel;
             }
             else
             {
-				Source.Source = UncheckedModel?.FlatModel;
+				Grid.DataContext = UncheckedModel?.FlatModel;
             }
         }
 
diff --git a/Engine/Source/Programs/Windows/CompileTimeAnalyzer/CompileTimeAnalyzer/Converters.cs b/Engine/Source/Programs/Windows/CompileTimeAnalyzer/CompileTimeAnalyzer/Converters.cs
--- a/Engine/Source/Programs/Windows/CompileTimeAnalyzer/CompileTimeAnalyzer/Converters.cs
+++ b/Engine/Source/Programs/Windows/CompileTimeAnalyzer/CompileTimeAnalyzer/Converters.cs
@@ -21,7 +21,7 @@ namespace Timing_Data_Investigator
     {
         public object Convert(object value, Type targetType, object parameter, CultureInfo culture)
         {
-            var VisibilityOption = parameter == null ? VisibilityOptions.VisibleOnTrue : Enum.Parse(typeof(VisibilityOptions), parameter.ToString());
+			VisibilityOptions VisibilityOption = parameter == null ? VisibilityOptions.VisibleOnTrue : (VisibilityOptions)Enum.Parse(typeof(VisibilityOptions), parameter.ToString());
             switch (VisibilityOption)
             {
                 case VisibilityOptions.VisibleOnFalse:
diff --git a/Engine/Source/Programs/Windows/CompileTimeAnalyzer/CompileTimeAnalyzer/MainWindow.xaml b/Engine/Source/Programs/Windows/CompileTimeAnalyzer/CompileTimeAnalyzer/MainWindow.xaml
--- a/Engine/Source/Programs/Windows/CompileTimeAnalyzer/CompileTimeAnalyzer/MainWindow.xaml
+++ b/Engine/Source/Programs/Windows/CompileTimeAnalyzer/CompileTimeAnalyzer/MainWindow.xaml
@@ -6,7 +6,7 @@
         xmlns:controls="clr-namespace:Timing_Data_Investigator.Controls"
         xmlns:tdi="clr-namespace:Timing_Data_Investigator"
         mc:Ignorable="d"
-        Title="Timing Data Investigator" Width="800" Height="600">
+        Title="CompileTimeAnalyzer" Width="800" Height="600">
 
     <Window.Resources>
 		<tdi:NullToVisibilityConverter x:Key="NullToVisibilityConverter" />
diff --git a/Engine/Source/Programs/Windows/CompileTimeAnalyzer/CompileTimeAnalyzer/Models/TimingDataViewModel.cs b/Engine/Source/Programs/Windows/CompileTimeAnalyzer/CompileTimeAnalyzer/Models/TimingDataViewModel.cs
--- a/Engine/Source/Programs/Windows/CompileTimeAnalyzer/CompileTimeAnalyzer/Models/TimingDataViewModel.cs
+++ b/Engine/Source/Programs/Windows/CompileTimeAnalyzer/CompileTimeAnalyzer/Models/TimingDataViewModel.cs
@@ -2,6 +2,7 @@
 
 using System.Collections.Generic;
 using System.Collections.ObjectModel;
+using System.ComponentModel;
 using System.IO;
 using System.IO.Compression;
 using System.Linq;
@@ -41,6 +42,79 @@ namespace Timing_Data_Investigator.Models
 
 		public ICommand OpenCommand { get; set; }
 
+		public int SortedIndex { get; set; }
+
+		public void UpdateSortIndex(string SortMember, ListSortDirection SortDirection)
+		{
+			int Index = 0;
+			UpdateSortIndexInternal(SortMember, SortDirection, this, ref Index);
+		}
+
+		private static void UpdateSortIndexInternal(string SortMember, ListSortDirection SortDirection, TimingDataViewModel ViewModel, ref int Index)
+		{
+			IEnumerable<TimingDataViewModel> ChildrenViewModels = ViewModel.Children.Cast<TimingDataViewModel>();
+			IOrderedEnumerable<TimingDataViewModel> SortedChildren = null;
+			switch (SortMember)
+			{
+				case "Name":
+					{
+						if (SortDirection == ListSortDirection.Ascending)
+						{
+							SortedChildren = ChildrenViewModels.OrderBy(c => c.Name);
+						}
+						else
+						{
+							SortedChildren = ChildrenViewModels.OrderByDescending(c => c.Name);
+						}
+						break;
+					}
+
+				case "InclusiveDuration":
+					{
+						if (SortDirection == ListSortDirection.Ascending)
+						{
+							SortedChildren = ChildrenViewModels.OrderBy(c => c.InclusiveDuration);
+						}
+						else
+						{
+							SortedChildren = ChildrenViewModels.OrderByDescending(c => c.InclusiveDuration);
+						}
+						break;
+					}
+
+				case "ExclusiveDuration":
+					{
+						if (SortDirection == ListSortDirection.Ascending)
+						{
+							SortedChildren = ChildrenViewModels.OrderBy(c => c.ExclusiveDuration);
+						}
+						else
+						{
+							SortedChildren = ChildrenViewModels.OrderByDescending(c => c.ExclusiveDuration);
+						}
+						break;
+					}
+			}
+
+			foreach (TimingDataViewModel Child in SortedChildren)
+			{
+				if (SortDirection == ListSortDirection.Ascending)
+				{
+					Child.SortedIndex = Index++;
+				}
+
+				if (Child.HasChildren)
+				{
+					UpdateSortIndexInternal(SortMember, SortDirection, Child, ref Index);
+				}
+
+				if (SortDirection == ListSortDirection.Descending)
+				{
+					Child.SortedIndex = Index++;
+				}
+			}
+		}
+
 		public static TimingDataViewModel FromBinaryFile(FileReference BinaryFile)
 		{
 			using (BinaryReader Reader = new BinaryReader(File.Open(BinaryFile.FullName, FileMode.Open, FileAccess.Read)))
diff --git a/Engine/Source/Programs/Windows/CompileTimeAnalyzer/CompileTimeAnalyzer/TreeGridFlatModel.cs b/Engine/Source/Programs/Windows/CompileTimeAnalyzer/CompileTimeAnalyzer/TreeGridFlatModel.cs
--- a/Engine/Source/Programs/Windows/CompileTimeAnalyzer/CompileTimeAnalyzer/TreeGridFlatModel.cs
+++ b/Engine/Source/Programs/Windows/CompileTimeAnalyzer/CompileTimeAnalyzer/TreeGridFlatModel.cs
@@ -4,6 +4,9 @@ using System;
 using System.Collections.Generic;
 using System.Collections.ObjectModel;
 using System.Collections.Specialized;
+using System.ComponentModel;
+using System.Linq;
+using System.Reflection;
 
 namespace Timing_Data_Investigator
 {
@@ -11,91 +14,156 @@ namespace Timing_Data_Investigator
 	{
 		private const string ModificationError = "The collection cannot be modified by the user.";
 
-		private bool modification;
-		private HashSet<TreeGridElement> keys;
+		private bool Modification;
+		private HashSet<TreeGridElement> Keys;
+		private PropertyInfo CurrentSortProperty;
+		private ListSortDirection? CurrentSortDirection;
 
 		public TreeGridFlatModel()
 		{
 			// Initialize the model
-			keys = new HashSet<TreeGridElement>();
+			Keys = new HashSet<TreeGridElement>();
+		}
+
+		public void Sort(PropertyInfo SortProperty, ListSortDirection? SortDirection)
+		{
+			// Set the modification flag
+			Modification = true;
+
+			IEnumerable<TreeGridElement> RootNodes = this.Where(element => element.Parent == null).ToList();
+			Items.Clear();
+			if (SortDirection == ListSortDirection.Ascending)
+			{
+				RootNodes = RootNodes.OrderBy(n => SortProperty.GetValue(n));
+			}
+			else
+			{
+				RootNodes = RootNodes.OrderByDescending(n => SortProperty.GetValue(n));
+			}
+
+			CurrentSortDirection = SortDirection;
+			CurrentSortProperty = SortProperty;
+
+			foreach (TreeGridElement Node in RootNodes.ToList())
+			{
+				SortNodes(Node);
+			}
+
+			OnPropertyChanged(new PropertyChangedEventArgs("Item[]"));
+			OnCollectionChanged(new NotifyCollectionChangedEventArgs(NotifyCollectionChangedAction.Reset));
+
+			// Clear the modification flag
+			Modification = false;
+		}
+
+		private void SortNodes(TreeGridElement Node)
+		{
+			Items.Add(Node);
+			if (Node.IsExpanded && Node.HasChildren)
+			{
+				IEnumerable<TreeGridElement> SortedChildren = Node.Children;
+				if (CurrentSortDirection == ListSortDirection.Ascending)
+				{
+					SortedChildren = Node.Children.OrderBy(c => CurrentSortProperty.GetValue(c));
+				}
+				else if (CurrentSortDirection == ListSortDirection.Descending)
+				{
+					SortedChildren = Node.Children.OrderByDescending(c => CurrentSortProperty.GetValue(c));
+				}
+
+				foreach (TreeGridElement Child in SortedChildren)
+				{
+					SortNodes(Child);
+				}
+			}
 		}
 
 		internal bool ContainsKey(TreeGridElement item)
 		{
 			// Return a value indicating if the item is within the model
-			return keys.Contains(item);
+			return Keys.Contains(item);
 		}
 
 		internal void PrivateInsert(int index, TreeGridElement item)
 		{
 			// Set the modification flag
-			modification = true;
+			Modification = true;
 
 			// Add the item to the model
 			Insert(index, item);
 
 			// Add the item to the keys
-			keys.Add(item);
+			Keys.Add(item);
 
 			// Clear the modification flag
-			modification = false;
+			Modification = false;
 		}
 
-		internal void PrivateInsertRange(int index, IList<TreeGridElement> items)
+		internal void PrivateInsertRange(int Index, IList<TreeGridElement> Items)
 		{
 			// Set the modification flag
-			modification = true;
+			Modification = true;
 
 			// Iterate through all of the children within the items
-			foreach (TreeGridElement child in items)
+			IEnumerable<TreeGridElement> SortedItems = Items;
+			if (CurrentSortDirection == ListSortDirection.Ascending)
+			{
+				Items.OrderBy(i => CurrentSortProperty.GetValue(i));
+			}
+			else if (CurrentSortDirection == ListSortDirection.Descending)
+			{
+				Items.OrderByDescending(i => CurrentSortProperty.GetValue(i));
+			}
+
+			foreach (TreeGridElement child in Items)
 			{
 				// Add the child to the model
-				Insert(index++, child);
+				Insert(Index++, child);
 
 				// Add the child to the keys
-				keys.Add(child);
+				Keys.Add(child);
 			}
 
 			// Clear the modification flag
-			modification = false;
+			Modification = false;
 		}
 
 		internal void PrivateRemoveRange(int index, int count)
 		{
 			// Set the modification flag
-			modification = true;
+			Modification = true;
 
 			// Iterate through all of the items to remove from the model
 			for (int itemIndex = 0; itemIndex < count; itemIndex++)
 			{
 				// Remove the item from the keys
-				keys.Remove(Items[index]);
+				Keys.Remove(Items[index]);
 
 				// Remove the item from the model
 				RemoveAt(index);
 			}
 
 			// Clear the modification flag
-			modification = false;
+			Modification = false;
 		}
 
 		internal void PrivateClear()
 		{
 			// Set the modification flag
-			modification = true;
+			Modification = true;
 
 			// Clear keys and items.
-			keys.Clear();
+			Keys.Clear();
 			Clear();
 
 			// Clear the modification flag
-			modification = false;
+			Modification = false;
 		}
 
 		protected override void OnCollectionChanged(NotifyCollectionChangedEventArgs args)
 		{
 			// Is the modification flag set?
-			if (!modification)
+			if (!Modification)
 			{
 				// The collection is for internal use only
 				throw new InvalidOperationException(ModificationError);
