Implementing CompileTimeAnalyzer feedback from Ben.

* Redesigned the grid to look better.
* Adjusted visible columns in different views.
* Double clicking files in an aggregate summary will open the file file.
* You can middle click tabs to close them.
* Fixed some timing data calculations that were causing exclusive/inclusive times to be inaccurate.

#rnx

diff --git a/Engine/Source/Programs/UnrealBuildTool/Modes/AggregateParsedTimingInfo.cs b/Engine/Source/Programs/UnrealBuildTool/Modes/AggregateParsedTimingInfo.cs
new file mode 100644
--- /dev/null
+++ b/Engine/Source/Programs/UnrealBuildTool/Modes/AggregateParsedTimingInfo.cs
@@ -0,0 +1,257 @@
+// Copyright 1998-2019 Epic Games, Inc. All Rights Reserved.
+
+using System;
+using System.Collections.Concurrent;
+using System.Collections.Generic;
+using System.IO;
+using System.IO.Compression;
+using System.Linq;
+using System.Text.RegularExpressions;
+using System.Threading;
+using System.Threading.Tasks;
+using Tools.DotNETCommon;
+
+namespace UnrealBuildTool
+{
+	/// <summary>
+	/// Aggregates parsed Visual C++ timing information files together into one monolithic breakdown file.
+	/// </summary>
+	[ToolMode("AggregateParsedTimingInfo", ToolModeOptions.None)]
+	class AggregateParsedTimingInfo : ToolMode
+	{
+		private TimingData FileTimingData;
+		private ConcurrentDictionary<string, int> DecompressedFileSizes = new ConcurrentDictionary<string, int>();
+		private ConcurrentDictionary<string, byte[]> CompressedFiles = new ConcurrentDictionary<string, byte[]>();
+		private ConcurrentBag<TimingData> AggregateIncludes = new ConcurrentBag<TimingData>();
+		private ConcurrentBag<TimingData> AggregateClasses = new ConcurrentBag<TimingData>();
+		private ConcurrentBag<TimingData> AggregateFunctions = new ConcurrentBag<TimingData>();
+		private object AddFileLock = new object();
+
+		public override int Execute(CommandLineArguments Arguments)
+		{
+			FileReference ManifestFile = Arguments.GetFileReference("-ManifestFile=");
+			string[] ParsedFileNames = FileReference.ReadAllLines(ManifestFile);
+
+			// Load all the file timing data and summarize them for the aggregate.
+			FileTimingData = new TimingData() { Name = "Files", Type = TimingDataType.Summary };
+			Parallel.ForEach(ParsedFileNames, new ParallelOptions() { MaxDegreeOfParallelism = 4 }, ParseTimingDataFile);
+
+			// Create aggregate summary. Duration is the duration of the files in the aggregate.
+			string AggregateName = Arguments.GetString("-Name=");
+			TimingData AggregateData = new TimingData() { Name = AggregateName, Type = TimingDataType.Aggregate };
+			AggregateData.AddChild(FileTimingData);
+
+			// Group the includes, classes, and functions by name and sum them up then add to the aggregate.
+			GroupTimingDataOnName(AggregateData, "Include Timings", AggregateIncludes);
+			GroupTimingDataOnName(AggregateData, "Class Timings", AggregateClasses);
+			GroupTimingDataOnName(AggregateData, "Function Timings", AggregateFunctions);
+
+			// Write out aggregate summary.
+			string OutputFile = Path.Combine(ManifestFile.Directory.FullName, String.Format("{0}.timing.bin", AggregateName));
+			using (BinaryWriter Writer = new BinaryWriter(File.Open(OutputFile, FileMode.Create)))
+			{
+				// Write out the aggregate data.
+				Writer.Write(AggregateData);
+
+				// Write the look up table for the compressed binary blobs.
+				int Offset = 0;
+				Writer.Write(CompressedFiles.Count);
+				foreach (KeyValuePair<string, byte[]> CompressedFile in CompressedFiles)
+				{
+					Writer.Write(CompressedFile.Key);
+					Writer.Write(Offset);
+					Writer.Write(CompressedFile.Value.Length);
+					Writer.Write(DecompressedFileSizes[CompressedFile.Key]);
+					Offset += CompressedFile.Value.Length;
+				}
+
+				// Write the compressed binary blobs.
+				foreach (KeyValuePair<string, byte[]> CompressedFile in CompressedFiles)
+				{
+					Writer.Write(CompressedFile.Value);
+				}
+			}
+
+			return 0;
+		}
+
+		private void GroupTimingDataOnName(TimingData AggregateData, string TimingName, IEnumerable<TimingData> UngroupedData)
+		{
+			string GroupName = TimingName;
+			TimingData GroupedTimingData = new TimingData() { Name = GroupName, Type = TimingDataType.Summary };
+			IEnumerable<IGrouping<string, TimingData>> Groups = UngroupedData.GroupBy(i => i.Name).OrderByDescending(g => g.Sum(d => d.ExclusiveDuration)).ToList();
+			foreach (IGrouping<string, TimingData> Group in Groups)
+			{
+				TimingData GroupedData = new TimingData() { Name = Group.Key, ExclusiveDuration = Group.Sum(d => d.ExclusiveDuration), Count = Group.Sum(d => d.Count) };
+				GroupedTimingData.Children.Add(Group.Key, GroupedData);
+			}
+
+			AggregateData.AddChild(GroupedTimingData);
+			AggregateData.ExclusiveDuration -= GroupedTimingData.InclusiveDuration;
+		}
+
+		private void ParseTimingDataFile(string ParsedFileName)
+		{
+			// Convert input file back into summary objects.
+			using (BinaryReader Reader = new BinaryReader(File.Open(ParsedFileName, FileMode.Open, FileAccess.Read)))
+			{
+				TimingData ParsedTimingData = new TimingData(Reader);
+
+				Task CompressDataTask = Task.Run(() =>
+				{
+					Reader.BaseStream.Seek(0, SeekOrigin.Begin);
+					using (MemoryStream CompressedMemoryStream = new MemoryStream())
+					{
+						using (GZipStream CompressionStream = new GZipStream(CompressedMemoryStream, CompressionMode.Compress))
+						{
+							Reader.BaseStream.CopyTo(CompressionStream);
+						}
+
+						CompressedFiles.TryAdd(ParsedTimingData.Name, CompressedMemoryStream.ToArray());
+						DecompressedFileSizes.TryAdd(ParsedTimingData.Name, (int)Reader.BaseStream.Length);
+					}
+				});
+
+				TimingData SummarizedTimingData = new TimingData() { Name = ParsedTimingData.Name, Parent = FileTimingData, Type = TimingDataType.Summary };
+				foreach (TimingData Include in ParsedTimingData.Children["IncludeTimings"].Children.Values)
+				{
+					SummarizedTimingData.AddChild(Include.Clone());
+				}
+				SummarizedTimingData.ExclusiveDuration += ParsedTimingData.Children["ClassTimings"].Children.Values.Sum(d => d.InclusiveDuration);
+				SummarizedTimingData.ExclusiveDuration += ParsedTimingData.Children["FunctionTimings"].Children.Values.Sum(d => d.InclusiveDuration);
+
+				// Gather and update child timing data.
+				Task QueueIncludesTask = Task.Run(() =>
+				{
+					// Flatten the includes for aggregation.
+					IEnumerable<TimingData> Includes = ParsedTimingData.Children["IncludeTimings"].Children.Values;
+					Dictionary<string, TimingData> FlattenedIncludes = new Dictionary<string, TimingData>();
+					FlattenIncludes(FlattenedIncludes, Includes);
+					foreach (TimingData Include in FlattenedIncludes.Values)
+					{
+						AggregateIncludes.Add(Include);
+					}
+				});
+
+				Task QueueClassesTask = Task.Run(() =>
+				{
+					// Collapse templates into single entries.
+					IEnumerable<TimingData> CollapsedClasses = GroupChildren(ParsedTimingData.Children["ClassTimings"].Children.Values, TimingDataType.Class);
+					foreach (TimingData Class in CollapsedClasses)
+					{
+						AggregateClasses.Add(new TimingData() { Name = Class.Name, Type = TimingDataType.Class, ExclusiveDuration = Class.InclusiveDuration });
+					}
+				});
+
+				Task QueueFunctionsTask = Task.Run(() =>
+				{
+					// Collapse templates into single entries.
+					IEnumerable<TimingData> CollapsedFunctions = GroupChildren(ParsedTimingData.Children["FunctionTimings"].Children.Values, TimingDataType.Function);
+					foreach (TimingData Function in CollapsedFunctions)
+					{
+						AggregateFunctions.Add(new TimingData() { Name = Function.Name, Type = TimingDataType.Function, ExclusiveDuration = Function.InclusiveDuration });
+					}
+				});
+
+				while (!CompressDataTask.IsCompleted || !QueueIncludesTask.IsCompleted || !QueueClassesTask.IsCompleted || !QueueFunctionsTask.IsCompleted)
+				{
+					Thread.Sleep(TimeSpan.FromMilliseconds(50));
+				}
+
+				lock (AddFileLock)
+				{
+					FileTimingData.AddChild(SummarizedTimingData);
+				}
+			}
+		}
+
+		private int HeaderFilesFound = 0;
+		private object CounterLock = new object();
+		private void FlattenIncludes(Dictionary<string, TimingData> Includes, IEnumerable<TimingData> IncludeData)
+		{
+			foreach (TimingData Include in IncludeData)
+			{
+				TimingData AggregatedInclude;
+				if (Includes.TryGetValue(Include.Name, out AggregatedInclude))
+				{
+					AggregatedInclude.ExclusiveDuration += Include.ExclusiveDuration;
+					AggregatedInclude.Count += 1;
+				}
+				else
+				{
+					if (Include.Name.EndsWith(".h"))
+					{
+						lock (CounterLock)
+						{
+							HeaderFilesFound += 1;
+						}
+					}
+
+					Includes.Add(Include.Name, new TimingData() { Name = Include.Name, Type = TimingDataType.Include, ExclusiveDuration = Include.ExclusiveDuration });
+				}
+
+				FlattenIncludes(Includes, Include.Children.Values);
+			}
+		}
+
+		private IEnumerable<TimingData> GroupChildren(IEnumerable<TimingData> Children, TimingDataType Type)
+		{
+			List<TimingData> GroupedChildren = new List<TimingData>();
+			Dictionary<string, List<TimingData>> ChildGroups = new Dictionary<string, List<TimingData>>();
+			foreach (TimingData Child in Children)
+			{
+				// See if this is a templated class. If not, add it as is.
+				Match Match = Regex.Match(Child.Name, @"^([^<]*)(?<Template><.*>)");
+				if (!Match.Success)
+				{
+					GroupedChildren.Add(Child);
+				}
+				else
+				{
+					// Generate group name from template.
+					int TemplateParamCount = Match.Groups["Template"].Value.Count(c => c == ',') + 1;
+					List<string> TemplateParamSig = new List<string>(TemplateParamCount);
+					for (int i = 0; i < TemplateParamCount; ++i)
+					{
+						TemplateParamSig.Add("...");
+					}
+					string GroupName = Child.Name.Replace(Match.Groups["Template"].Value, String.Format("<{0}>", string.Join(", ", TemplateParamSig)));
+
+					// See if we have a group for this template already. If not, add it.
+					if (!ChildGroups.ContainsKey(GroupName))
+					{
+						ChildGroups.Add(GroupName, new List<TimingData>());
+					}
+
+					ChildGroups[GroupName].Add(Child);
+				}
+			}
+
+			// Add grouped children.
+			foreach (KeyValuePair<string, List<TimingData>> Group in ChildGroups)
+			{
+				if (Group.Value.Count == 1)
+				{
+					GroupedChildren.Add(Group.Value.First());
+					continue;
+				}
+
+				TimingData NewViewModel = new TimingData()
+				{
+					Name = Group.Key,
+					Type = Type,
+				};
+
+				foreach (TimingData Child in Group.Value)
+				{
+					TimingData NewChild = new TimingData() { Name = Child.Name, Type = Child.Type, Parent = NewViewModel, ExclusiveDuration = Child.ExclusiveDuration };
+					NewViewModel.AddChild(NewChild);
+				}
+
+				GroupedChildren.Add(NewViewModel);
+			}
+
+			return GroupedChildren;
+		}
+	}
+}
diff --git a/Engine/Source/Programs/UnrealBuildTool/System/TimingData.cs b/Engine/Source/Programs/UnrealBuildTool/System/TimingData.cs
--- a/Engine/Source/Programs/UnrealBuildTool/System/TimingData.cs
+++ b/Engine/Source/Programs/UnrealBuildTool/System/TimingData.cs
@@ -67,7 +67,8 @@ namespace UnrealBuildTool
 		{
 			Name = Reader.ReadString();
 			Type = (TimingDataType)Reader.ReadByte();
-			ExclusiveDuration = (float)Reader.ReadDouble();
+			Count = Reader.ReadInt32();
+			ExclusiveDuration = Reader.ReadDouble();
 			int ChildCount = Reader.ReadInt32();
 			for (int i = 0; i < ChildCount; ++i)
 			{
@@ -86,21 +87,20 @@ namespace UnrealBuildTool
 		/// </summary>
 		public TimingDataType Type { get; set; }
 
+		/// <summary>
+		/// Gets the number of times this event was encountered.
+		/// </summary>
+		public int Count { get; set; } = 1;
+
 		/// <summary>
 		/// Gets the amount of time, in milliseconds, that the event took to complete disregarding its sub-events.
 		/// </summary>
-		public double ExclusiveDuration { get; set; }
+		public double ExclusiveDuration { get; set; } = 0.0;
 
 		/// <summary>
 		/// Gets the amount of time, in milliseconds, that the event took including its sub-events.
 		/// </summary>
-		public double InclusiveDuration
-		{
-			get
-			{
-				return ExclusiveDuration + Children.Sum(c => c.Value.InclusiveDuration);
-			}
-		}
+		public double InclusiveDuration => ExclusiveDuration + Children.Sum(c => c.Value.InclusiveDuration);
 
 		/// <summary>
 		/// If this event is a sub-event, the parent is the event it is the sub-event to.
@@ -163,6 +163,7 @@ namespace UnrealBuildTool
 		{
 			Writer.Write(Name);
 			Writer.Write((byte)Type);
+			Writer.Write(Count);
 			Writer.Write(ExclusiveDuration);
 			Writer.Write(Children.Count);
 			foreach (TimingData Child in Children.Values)
diff --git a/Engine/Source/Programs/UnrealBuildTool/UnrealBuildTool.csproj b/Engine/Source/Programs/UnrealBuildTool/UnrealBuildTool.csproj
--- a/Engine/Source/Programs/UnrealBuildTool/UnrealBuildTool.csproj
+++ b/Engine/Source/Programs/UnrealBuildTool/UnrealBuildTool.csproj
@@ -187,14 +187,14 @@
     <Compile Include="ProjectFiles\Xcode\XcodeProjectFileGenerator.cs" />
     <Compile Include="System\Action.cs" />
     <Compile Include="System\ActionExecutor.cs" />
-    <Compile Include="System\AggregateParsedTimingInfo.cs" />
+    <Compile Include="Modes\AggregateParsedTimingInfo.cs" />
     <Compile Include="System\Asn.cs" />
     <Compile Include="System\BuildVersion.cs" />
     <Compile Include="System\CppCompileEnvironment.cs" />
     <Compile Include="System\CppDependencyCache.cs" />
     <Compile Include="System\CppIncludeParser.cs" />
     <Compile Include="System\FileMetadataPrefetch.cs" />
-    <Compile Include="System\ParseMsvcTimingInfo.cs" />
+    <Compile Include="Modes\ParseMsvcTimingInfo.cs" />
     <Compile Include="System\RulesScope.cs" />
     <Compile Include="System\SourceFileMetadataCache.cs" />
     <Compile Include="System\DataDrivenPlatformInfo.cs" />
diff --git a/Engine/Source/Programs/Windows/CompileTimeAnalyzer/CompileTimeAnalyzer/Controls/TimingDataGrid.xaml b/Engine/Source/Programs/Windows/CompileTimeAnalyzer/CompileTimeAnalyzer/Controls/TimingDataGrid.xaml
--- a/Engine/Source/Programs/Windows/CompileTimeAnalyzer/CompileTimeAnalyzer/Controls/TimingDataGrid.xaml
+++ b/Engine/Source/Programs/Windows/CompileTimeAnalyzer/CompileTimeAnalyzer/Controls/TimingDataGrid.xaml
@@ -6,69 +6,144 @@
              xmlns:local="clr-namespace:Timing_Data_Investigator.Controls"
              xmlns:tdi="clr-namespace:Timing_Data_Investigator"
              mc:Ignorable="d" >
-    
-    <UserControl.Resources>
-        <GridLength x:Key="TreeToggleWidth">20</GridLength>
-        <tdi:LevelConverter x:Key="LevelSpaceConverter" LevelWidth="{StaticResource TreeToggleWidth}"/>
-        <tdi:BoolToVisibilityConverter x:Key="BoolToVisibilityConverter" />
-        <tdi:MultiBoolToVisibilityConverter x:Key="MultiBoolToVisibilityConverter" />
+
+	<UserControl.Resources>
+		<GridLength x:Key="TreeToggleWidth">20</GridLength>
+		<tdi:LevelConverter x:Key="LevelSpaceConverter" LevelWidth="{StaticResource TreeToggleWidth}"/>
+		<tdi:BoolToVisibilityConverter x:Key="BoolToVisibilityConverter" />
+		<tdi:MultiBoolToVisibilityConverter x:Key="MultiBoolToVisibilityConverter" />
 		<tdi:IntLessThanConverter x:Key="IntComparisonConverter" />
 
 		<Style TargetType="{x:Type CheckBox}" x:Key="Expander">
 			<Setter Property="Template">
-                <Setter.Value>
-                    <ControlTemplate TargetType="{x:Type CheckBox}">
-                        <StackPanel Orientation="Horizontal">
+				<Setter.Value>
+					<ControlTemplate TargetType="{x:Type CheckBox}">
+						<StackPanel Orientation="Horizontal">
 							<Canvas>
 								<Polygon x:Name="Expander" Points="3,3 9,8 3,13" Fill="Black" />
 							</Canvas>
 							<ContentPresenter/>
-                        </StackPanel>
-                        <ControlTemplate.Triggers>
-                            <Trigger Property="IsChecked" Value="True">
-                                <Setter TargetName="Expander" Property="Points" Value="3,4 13,4 8,11" />
-                            </Trigger>
-                            <Trigger Property="IsMouseOver" Value="True">
-                                <Setter TargetName="Expander" Property="Stroke" Value="LightGray"/>
-                                <Setter TargetName="Expander" Property="Fill" Value="DarkGray"/>
-                            </Trigger>
-                        </ControlTemplate.Triggers>
-                    </ControlTemplate>
-                </Setter.Value>
-            </Setter>
-        </Style>
-    </UserControl.Resources>
+						</StackPanel>
+						<ControlTemplate.Triggers>
+							<Trigger Property="IsChecked" Value="True">
+								<Setter TargetName="Expander" Property="Points" Value="3,13 10,5 10,13" />
+								<Setter TargetName="Expander" Property="Stroke" Value="LightGray"/>
+								<Setter TargetName="Expander" Property="Fill" Value="DarkGray"/>
+							</Trigger>
+						</ControlTemplate.Triggers>
+					</ControlTemplate>
+				</Setter.Value>
+			</Setter>
+		</Style>
+
+		<ControlTemplate x:Key="ProgressBarControlTemplate" TargetType="{x:Type ProgressBar}">
+			<Grid x:Name="TemplateRoot">
+				<Border BorderBrush="{TemplateBinding BorderBrush}" BorderThickness="{TemplateBinding BorderThickness}" Background="{TemplateBinding Background}" CornerRadius="3">
+					<Grid x:Name="PART_Indicator" ClipToBounds="True" HorizontalAlignment="Left">
+						<Rectangle x:Name="Indicator" Fill="{TemplateBinding Foreground}" RadiusX="3" RadiusY="3" />
+					</Grid>
+				</Border>
+				<Rectangle x:Name="PART_Track"/>
+			</Grid>
+			<ControlTemplate.Triggers>
+				<Trigger Property="Orientation" Value="Vertical">
+					<Setter Property="LayoutTransform" TargetName="TemplateRoot">
+						<Setter.Value>
+							<RotateTransform Angle="-90"/>
+						</Setter.Value>
+					</Setter>
+				</Trigger>
+				<Trigger Property="IsIndeterminate" Value="True">
+					<Setter Property="Visibility" TargetName="Indicator" Value="Collapsed"/>
+				</Trigger>
+			</ControlTemplate.Triggers>
+		</ControlTemplate>
+	</UserControl.Resources>
 
-	<DataGrid AutoGenerateColumns="False" Name="Grid" Grid.Row="1" CanUserAddRows="False" AlternationCount="2" AlternatingRowBackground="#EEE"
-              ItemsSource="{Binding}" Sorting="Grid_Sorting" PreviewKeyDown="Grid_KeyDown">
+	<DataGrid AutoGenerateColumns="False" x:Name="Grid" Grid.Row="1" CanUserAddRows="False" CanUserDeleteRows="False" CanUserResizeRows="False" ItemsSource="{Binding}"
+			  Sorting="Grid_Sorting" PreviewKeyDown="Grid_KeyDown" HorizontalGridLinesBrush="#F4F4F4" VerticalGridLinesBrush="#F4F4F4" SelectionUnit="FullRow" SelectionMode="Single" HeadersVisibility="Column">
 		<DataGrid.Columns>
-            <DataGridTemplateColumn Header="Name" Width="*" CanUserSort="True" SortMemberPath="Name">
-                <DataGridTemplateColumn.CellTemplate>
-                    <DataTemplate>
-                        <Grid>
-                            <Grid.ColumnDefinitions>
-                                <ColumnDefinition Width="{Binding Level, Converter={StaticResource LevelSpaceConverter}}"/>
-                                <ColumnDefinition Width="{StaticResource TreeToggleWidth}"/>
-                                <ColumnDefinition Width="*"/>
-                            </Grid.ColumnDefinitions>
+			<DataGridTemplateColumn Header="Name" Width="*" CanUserSort="True" SortMemberPath="Name">
+				<DataGridTemplateColumn.CellTemplate>
+					<DataTemplate>
+						<Grid>
+							<Grid.ColumnDefinitions>
+								<ColumnDefinition Width="{Binding Level, Converter={StaticResource LevelSpaceConverter}}"/>
+								<ColumnDefinition Width="{StaticResource TreeToggleWidth}"/>
+								<ColumnDefinition Width="*"/>
+							</Grid.ColumnDefinitions>
 
-                            <CheckBox Grid.Column="1" IsChecked="{Binding IsExpanded, Mode=TwoWay, UpdateSourceTrigger=PropertyChanged}" Visibility="{Binding HasChildren, Converter={StaticResource BoolToVisibilityConverter}}" Style="{StaticResource Expander}" />
-                            <Line Grid.Column="1" X1="8" X2="8" Y2="{Binding ActualHeight, RelativeSource={RelativeSource Self}}" Stroke="#888" StrokeDashArray="2 2" StrokeDashOffset="1">
-                                <Line.Visibility>
-                                    <MultiBinding Converter="{StaticResource MultiBoolToVisibilityConverter}" ConverterParameter="VisibleOnFalse">
-                                        <Binding Path="HasChildren" />
+							<Button Grid.Column="1" Background="Transparent" BorderThickness="0" Click="Button_Click">
+								<Button.Style>
+									<Style TargetType="Button">
+										<Setter Property="OverridesDefaultStyle" Value="True" />
+										<Setter Property="Template">
+											<Setter.Value>
+												<ControlTemplate TargetType="Button">
+													<Border Name="border" BorderThickness="0" BorderBrush="Black" Background="{TemplateBinding Background}">
+														<ContentPresenter HorizontalAlignment="Center" VerticalAlignment="Center" />
+													</Border>
+												</ControlTemplate>
+											</Setter.Value>
+										</Setter>
+									</Style>
+								</Button.Style>
+							</Button>
+							<CheckBox Grid.Column="1" IsChecked="{Binding IsExpanded, Mode=TwoWay, UpdateSourceTrigger=PropertyChanged}" Visibility="{Binding HasChildren, Converter={StaticResource BoolToVisibilityConverter}}" Style="{StaticResource Expander}" />
+							<Line Grid.Column="1" X1="8" X2="8" Y2="{Binding ActualHeight, RelativeSource={RelativeSource Self}}" Stroke="#888" StrokeDashArray="2 2" StrokeDashOffset="1">
+								<Line.Visibility>
+									<MultiBinding Converter="{StaticResource MultiBoolToVisibilityConverter}" ConverterParameter="VisibleOnFalse">
+										<Binding Path="HasChildren" />
 										<Binding Path="Level" Converter="{StaticResource IntComparisonConverter}" ConverterParameter="1" />
-                                    </MultiBinding>
-                                </Line.Visibility>
-                            </Line>
-                            <TextBlock Grid.Column="2" Text="{Binding Name}"/>
-                        </Grid>
-                    </DataTemplate>
-                </DataGridTemplateColumn.CellTemplate>
-            </DataGridTemplateColumn>
-			<DataGridTextColumn Header="Duration" Binding="{Binding InclusiveDuration, StringFormat={}{0:0.########}s}" SortMemberPath="InclusiveDuration" />
-			<DataGridTextColumn Header="Duration (Exclusive)" Binding="{Binding ExclusiveDuration, StringFormat={}{0:0.########}s}" SortMemberPath="ExclusiveDuration" />
-        </DataGrid.Columns>
-    </DataGrid>
+									</MultiBinding>
+								</Line.Visibility>
+							</Line>
+							<TextBlock Grid.Column="2" Text="{Binding Name}"/>
+						</Grid>
+					</DataTemplate>
+				</DataGridTemplateColumn.CellTemplate>
+			</DataGridTemplateColumn>
+
+			<DataGridTextColumn x:Name="CountColumn" Header="Count" Binding="{Binding Count}" SortMemberPath="Count" />
+
+			<DataGridTemplateColumn Header="Duration" CanUserSort="True" SortMemberPath="InclusiveDuration">
+				<DataGridTemplateColumn.CellTemplate>
+					<DataTemplate>
+						<Grid>
+							<ProgressBar Minimum="0" Maximum="100" Value="{Binding InclusivePercent, Mode=OneWay}" FlowDirection="RightToLeft" Template="{DynamicResource ProgressBarControlTemplate}"
+										 Margin="2,1,2,1" Foreground="CornflowerBlue" Background="Transparent" BorderThickness="0" />
+							<TextBlock TextAlignment="Right" Margin="4,0,4,0">
+								<TextBlock.Text>
+									<MultiBinding StringFormat="{}{0:0.###}s ({1:0.##}%)">
+										<Binding Path="InclusiveDuration" />
+										<Binding Path="InclusivePercent" />
+									</MultiBinding>
+								</TextBlock.Text>
+							</TextBlock>
+						</Grid>
+					</DataTemplate>
+				</DataGridTemplateColumn.CellTemplate>
+			</DataGridTemplateColumn>
+
+			<DataGridTemplateColumn x:Name="ExclusiveDurationColumn" Header="Duration (Exclusive)" CanUserSort="True" SortMemberPath="ExclusiveDuration">
+				<DataGridTemplateColumn.CellTemplate>
+					<DataTemplate>
+						<Grid>
+							<ProgressBar Minimum="0" Maximum="100" Value="{Binding ExclusivePercent, Mode=OneWay}" FlowDirection="RightToLeft" Template="{DynamicResource ProgressBarControlTemplate}"
+										 Margin="2,1,2,1" Foreground="CornflowerBlue" Background="Transparent" BorderThickness="0" />
+							<TextBlock TextAlignment="Right" Margin="4,0,4,0">
+								<TextBlock.Text>
+									<MultiBinding StringFormat="{}{0:0.###}s ({1:0.##}%)">
+										<Binding Path="ExclusiveDuration" />
+										<Binding Path="ExclusivePercent" />
+									</MultiBinding>
+								</TextBlock.Text>
+							</TextBlock>
+						</Grid>
+					</DataTemplate>
+				</DataGridTemplateColumn.CellTemplate>
+			</DataGridTemplateColumn>
+		</DataGrid.Columns>
+	</DataGrid>
 </UserControl>
  
\ No newline at end of file
diff --git a/Engine/Source/Programs/Windows/CompileTimeAnalyzer/CompileTimeAnalyzer/Controls/TimingDataGrid.xaml.cs b/Engine/Source/Programs/Windows/CompileTimeAnalyzer/CompileTimeAnalyzer/Controls/TimingDataGrid.xaml.cs
--- a/Engine/Source/Programs/Windows/CompileTimeAnalyzer/CompileTimeAnalyzer/Controls/TimingDataGrid.xaml.cs
+++ b/Engine/Source/Programs/Windows/CompileTimeAnalyzer/CompileTimeAnalyzer/Controls/TimingDataGrid.xaml.cs
@@ -1,12 +1,10 @@
 // Copyright 1998-2019 Epic Games, Inc. All Rights Reserved.
 
-using System.Collections.Generic;
 using System.Collections.ObjectModel;
 using System.ComponentModel;
 using System.Linq;
 using System.Reflection;
 using System.Windows.Controls;
-using System.Windows.Data;
 using System.Windows.Input;
 using Timing_Data_Investigator.Models;
 
@@ -19,13 +17,11 @@ namespace Timing_Data_Investigator.Controls
     {
         public TimingDataGrid()
         {
-            InitializeComponent();
+			InitializeComponent();
         }
 
         private void Grid_Sorting(object sender, DataGridSortingEventArgs e)
         {
-			
-			
 			switch (e.Column.SortDirection)
 			{
 				case ListSortDirection.Ascending:
@@ -108,5 +104,16 @@ namespace Timing_Data_Investigator.Controls
                 }
             }
         }
-    }
+
+		private void Button_Click(object sender, System.Windows.RoutedEventArgs e)
+		{
+			TreeGridElement SelectedData = Grid.SelectedItem as TreeGridElement;
+			if (SelectedData == null)
+			{
+				return;
+			}
+
+			SelectedData.IsExpanded = !SelectedData.IsExpanded;
+		}
+	}
 }
diff --git a/Engine/Source/Programs/Windows/CompileTimeAnalyzer/CompileTimeAnalyzer/Controls/TimingDataTab.xaml b/Engine/Source/Programs/Windows/CompileTimeAnalyzer/CompileTimeAnalyzer/Controls/TimingDataTab.xaml
--- a/Engine/Source/Programs/Windows/CompileTimeAnalyzer/CompileTimeAnalyzer/Controls/TimingDataTab.xaml
+++ b/Engine/Source/Programs/Windows/CompileTimeAnalyzer/CompileTimeAnalyzer/Controls/TimingDataTab.xaml
@@ -11,13 +11,15 @@
         <ContextMenu  x:Key="RowMenu" DataContext="{Binding PlacementTarget.DataContext, RelativeSource={RelativeSource Self}}">
             <MenuItem Header="{Binding Name, StringFormat=Open {0} in New Tab}" Command="{Binding OpenCommand}" />
         </ContextMenu>
-        <Style x:Key="FilesRowStyle" TargetType="{x:Type DataGridRow}">
-            <Setter Property="ContextMenu" Value="{StaticResource RowMenu}" />
+		<Style x:Key="FilesRowStyle" TargetType="{x:Type DataGridRow}">
+			<EventSetter Event="MouseDoubleClick" Handler="DataGridRow_MouseDoubleClick" />
+			<EventSetter Event="MouseEnter" Handler="DataGridRow_MouseEnter" />
+			<EventSetter Event="MouseLeave" Handler="DataGridRow_MouseLeave" />
         </Style>
     </UserControl.Resources>
 
     <TabControl x:Name="SummaryTabs">
-        <TabItem x:Name="FilesTab" Header="Files" Visibility="Collapsed">
+		<TabItem x:Name="FilesTab" Header="Files" Visibility="Collapsed">
             <local:TimingDataGrid x:Name="FilesGrid" DockPanel.Dock="Top" />
         </TabItem>
         <TabItem x:Name="IncludesTab" Header="Includes">
@@ -26,11 +28,8 @@
 					<RowDefinition Height="Auto" />
 					<RowDefinition Height="*" />
 				</Grid.RowDefinitions>
-				<StackPanel Grid.Row="0" x:Name="FlattenIncludesRow" Orientation="Horizontal">
-					<TextBlock Text="Flatten Includes:" VerticalAlignment="Center" Margin="4 2" />
-					<CheckBox x:Name="FlattenIncludes" Checked="FlattenIncludes_Checked" Unchecked="FlattenIncludes_Checked" VerticalAlignment="Center" Margin="2"/>
-				</StackPanel>
-
+				
+				<CheckBox x:Name="FlattenIncludes" Checked="FlattenIncludes_Checked" Unchecked="FlattenIncludes_Checked" VerticalAlignment="Center" Margin="4" Content="Flatten Includes" />
 				<local:TimingDataGrid Grid.Row="1" x:Name="IncludesGrid" DockPanel.Dock="Top" />
 			</Grid>
 		</TabItem>
@@ -40,11 +39,8 @@
                     <RowDefinition Height="Auto" />
                     <RowDefinition Height="*" />
                 </Grid.RowDefinitions>
-				<StackPanel Grid.Row="0" x:Name="GroupClassesRow" Orientation="Horizontal">
-					<TextBlock Text="Group Template Classes:" VerticalAlignment="Center" Margin="4 2" />
-					<CheckBox x:Name="GroupClassTemplates" Checked="GroupTemplates_Checked" Unchecked="GroupTemplates_Checked" VerticalAlignment="Center" Margin="2" />
-                </StackPanel>
-
+				
+				<CheckBox x:Name="GroupClassTemplates" Checked="GroupTemplates_Checked" Unchecked="GroupTemplates_Checked" VerticalAlignment="Center" Margin="4" Content="Group Template Classes" />
                 <local:TimingDataGrid Grid.Row="1" x:Name="ClassesGrid" />
             </Grid>
         </TabItem>
@@ -54,11 +50,8 @@
                     <RowDefinition Height="Auto" />
                     <RowDefinition Height="*" />
                 </Grid.RowDefinitions>
-                <StackPanel Grid.Row="0" x:Name="GroupFunctionsRow" Orientation="Horizontal">
-					<TextBlock Text="Group Template Classes:" VerticalAlignment="Center" Margin="4 2" />
-					<CheckBox x:Name="GroupFunctionTemplates" Checked="GroupTemplates_Checked" Unchecked="GroupTemplates_Checked" VerticalAlignment="Center" Margin="2" />
-                </StackPanel>
-
+				
+				<CheckBox x:Name="GroupFunctionTemplates" Checked="GroupTemplates_Checked" Unchecked="GroupTemplates_Checked" VerticalAlignment="Center" Margin="4" Content="Group Template Functions" />
                 <local:TimingDataGrid Grid.Row="1" x:Name="FunctionsGrid" />
             </Grid>
         </TabItem>
diff --git a/Engine/Source/Programs/Windows/CompileTimeAnalyzer/CompileTimeAnalyzer/Controls/TimingDataTab.xaml.cs b/Engine/Source/Programs/Windows/CompileTimeAnalyzer/CompileTimeAnalyzer/Controls/TimingDataTab.xaml.cs
--- a/Engine/Source/Programs/Windows/CompileTimeAnalyzer/CompileTimeAnalyzer/Controls/TimingDataTab.xaml.cs
+++ b/Engine/Source/Programs/Windows/CompileTimeAnalyzer/CompileTimeAnalyzer/Controls/TimingDataTab.xaml.cs
@@ -1,14 +1,13 @@
 // Copyright 1998-2019 Epic Games, Inc. All Rights Reserved.
 
 using System.Collections.Generic;
-using System.Collections.ObjectModel;
 using System.ComponentModel;
 using System.Linq;
 using System.Reflection;
 using System.Text.RegularExpressions;
 using System.Windows;
 using System.Windows.Controls;
-using System.Windows.Data;
+using System.Windows.Media;
 using Timing_Data_Investigator.Models;
 using UnrealBuildTool;
 
@@ -100,10 +99,7 @@ namespace Timing_Data_Investigator.Controls
 				UpdateUIForSingleFile(NewTimingData);
             }
 
-            SummaryTabs.SelectedItem = FilesModel == null ? IncludesTab : FilesTab;
-			FlattenIncludesRow.Visibility = FlattenedIncludesModel != null ? Visibility.Visible : Visibility.Hidden;
-			GroupClassesRow.Visibility = GroupedClassesModel != null ? Visibility.Visible : Visibility.Hidden;
-			GroupFunctionsRow.Visibility = GroupedFunctionsModel != null ? Visibility.Visible : Visibility.Hidden;
+			SummaryTabs.SelectedItem = FilesModel == null ? IncludesTab : FilesTab;
 
 			if (TabStates.TryGetValue(NewTimingData, out TabState State))
 			{
@@ -118,20 +114,27 @@ namespace Timing_Data_Investigator.Controls
 				RefreshTabState(IncludesTab, State);
 				RefreshTabState(ClassesTab, State);
 				RefreshTabState(FunctionsTab, State);
-
-				SortModel(FilesGrid.Grid);
-				SortModel(IncludesGrid.Grid);
-				SortModel(ClassesGrid.Grid);
-				SortModel(FunctionsGrid.Grid);
 			}
-        }
+			else
+			{
+				FilesGrid.Grid.Columns[0].SortDirection = ListSortDirection.Ascending;
+				IncludesGrid.Grid.Columns[2].SortDirection = ListSortDirection.Descending;
+				ClassesGrid.Grid.Columns[2].SortDirection = ListSortDirection.Descending;
+				FunctionsGrid.Grid.Columns[2].SortDirection = ListSortDirection.Descending;
+			}
+
+			SortModel(FilesGrid.Grid);
+			SortModel(IncludesGrid.Grid);
+			SortModel(ClassesGrid.Grid);
+			SortModel(FunctionsGrid.Grid);
+		}
 
 		private void UpdateUIForAggregate(TimingDataViewModel TimingData)
 		{
-			FilesModel = GenerateTreeGridModel(TimingData.Children[0].Children);
-			IncludesModel = GenerateTreeGridModel(TimingData.Children[1].Children);
-			ClassesModel = GenerateTreeGridModel(TimingData.Children[2].Children);
-			FunctionsModel = GenerateTreeGridModel(TimingData.Children[3].Children);
+			FilesModel = GenerateTreeGridModel(TimingData.Children[0]);
+			IncludesModel = GenerateTreeGridModel(TimingData.Children[1]);
+			ClassesModel = GenerateTreeGridModel(TimingData.Children[2]);
+			FunctionsModel = GenerateTreeGridModel(TimingData.Children[3]);
 			FlattenedIncludesModel = null;
 			GroupedClassesModel = null;
 			GroupedFunctionsModel = null;
@@ -140,14 +143,24 @@ namespace Timing_Data_Investigator.Controls
 			UpdateTabVisibility(IncludesModel.FlatModel, IncludesTab, IncludesGrid);
 			UpdateTabVisibility(ClassesModel.FlatModel, ClassesTab, ClassesGrid);
 			UpdateTabVisibility(FunctionsModel.FlatModel, FunctionsTab, FunctionsGrid);
+
+			FlattenIncludes.IsEnabled = false;
+			GroupClassTemplates.IsEnabled = false;
+			GroupFunctionTemplates.IsEnabled = false;
+
+			FilesGrid.CountColumn.Visibility = Visibility.Collapsed;
+			IncludesGrid.CountColumn.Visibility = Visibility.Visible;
+			IncludesGrid.ExclusiveDurationColumn.Visibility = Visibility.Collapsed;
+			ClassesGrid.ExclusiveDurationColumn.Visibility = Visibility.Collapsed;
+			FunctionsGrid.ExclusiveDurationColumn.Visibility = Visibility.Collapsed;
 		}
 
 		private void UpdateUIForSingleFile(TimingDataViewModel TimingData)
 		{
 			FilesModel = null;
-			IncludesModel = GenerateTreeGridModel(TimingData.Children[0].Children);
-			ClassesModel = GenerateTreeGridModel(TimingData.Children[1].Children);
-			FunctionsModel = GenerateTreeGridModel(TimingData.Children[2].Children);
+			IncludesModel = GenerateTreeGridModel(TimingData.Children[0]);
+			ClassesModel = GenerateTreeGridModel(TimingData.Children[1]);
+			FunctionsModel = GenerateTreeGridModel(TimingData.Children[2]);
 			FlattenedIncludesModel = GenerateFlattenedModel(IncludesModel);
 			GroupedClassesModel = GenerateGroupedModel(ClassesModel);
 			GroupedFunctionsModel = GenerateGroupedModel(FunctionsModel);
@@ -156,6 +169,15 @@ namespace Timing_Data_Investigator.Controls
 			UpdateTabVisibility(FlattenIncludes.IsChecked == true ? FlattenedIncludesModel.FlatModel : IncludesModel.FlatModel, IncludesTab, IncludesGrid);
 			UpdateTabVisibility(GroupClassTemplates.IsChecked == true ? GroupedClassesModel.FlatModel : ClassesModel.FlatModel, ClassesTab, ClassesGrid);
 			UpdateTabVisibility(GroupFunctionTemplates.IsChecked == true ? GroupedFunctionsModel.FlatModel : FunctionsModel.FlatModel, FunctionsTab, FunctionsGrid);
+
+			FlattenIncludes.IsEnabled = true;
+			GroupClassTemplates.IsEnabled = true;
+			GroupFunctionTemplates.IsEnabled = true;
+
+			IncludesGrid.CountColumn.Visibility = Visibility.Collapsed;
+			IncludesGrid.ExclusiveDurationColumn.Visibility = Visibility.Visible;
+			ClassesGrid.ExclusiveDurationColumn.Visibility = Visibility.Visible;
+			FunctionsGrid.ExclusiveDurationColumn.Visibility = Visibility.Visible;
 		}
 
 		private void RefreshTabState(TabItem Tab, TabState State)
@@ -194,18 +216,19 @@ namespace Timing_Data_Investigator.Controls
 			Grid.DataContext = TabModel;
         }
 
-		private TreeGridModel GenerateTreeGridModel(IEnumerable<TreeGridElement> Children)
+		private TreeGridModel GenerateTreeGridModel(TreeGridElement Parent)
 		{
-			TreeGridModel Model = new TreeGridModel();
-			foreach (TreeGridElement Child in Children)
+			// Make sure to set the parent duration override since these children will get de-parented.
+			double ParentDuration = ((TimingDataViewModel)Parent).InclusiveDuration;
+			foreach (TimingDataViewModel Child in Parent.Children.Cast<TimingDataViewModel>())
 			{
-				Model.Add(Child);
+				Child.ParentDurationOverride = ParentDuration;
 			}
 
-			return Model;
+			return new TreeGridModel(Parent.Children);
 		}
 
-		private void FlattenIncludesInternal(Dictionary<string, TimingDataViewModel> FlattenedIncludes, IEnumerable<TimingDataViewModel> ViewModels)
+		private void FlattenIncludesInternal(Dictionary<string, TimingDataViewModel> FlattenedIncludes, IEnumerable<TimingDataViewModel> ViewModels, double ParentDurationOverride)
 		{
 			foreach (TimingDataViewModel TimingData in ViewModels)
 			{
@@ -219,34 +242,32 @@ namespace Timing_Data_Investigator.Controls
 					{
 						Name = TimingData.Name,
 						Type = TimingData.Type,
+						ParentDurationOverride = ParentDurationOverride,
 						ExclusiveDuration = TimingData.ExclusiveDuration
 					};
 
 					FlattenedIncludes.Add(FlattenedInclude.Name, FlattenedInclude);
 				}
 
-				FlattenIncludesInternal(FlattenedIncludes, TimingData.Children.Cast<TimingDataViewModel>());
+				FlattenIncludesInternal(FlattenedIncludes, TimingData.Children.Cast<TimingDataViewModel>(), ParentDurationOverride);
 			}
 		}
 
 		private TreeGridModel GenerateFlattenedModel(TreeGridModel UnflattenedModel)
 		{
+			IEnumerable<TimingDataViewModel> TimingDataModel = UnflattenedModel.Cast<TimingDataViewModel>();
+			double OverrideDuration = TimingDataModel.ElementAt(0).ParentDurationOverride ?? TimingDataModel.Sum(d => d.InclusiveDuration);
 			Dictionary<string, TimingDataViewModel> FlattenedIncludes = new Dictionary<string, TimingDataViewModel>();
-			FlattenIncludesInternal(FlattenedIncludes, UnflattenedModel.Cast<TimingDataViewModel>());
-			TreeGridModel FlattenedModel = new TreeGridModel();
-			foreach (TimingDataViewModel TimingData in FlattenedIncludes.Values)
-			{
-				FlattenedModel.Add(TimingData);
-			}
-
-			return FlattenedModel;
+			FlattenIncludesInternal(FlattenedIncludes, TimingDataModel, OverrideDuration);
+			return new TreeGridModel(FlattenedIncludes.Values);
 		}
 
 		private TreeGridModel GenerateGroupedModel(TreeGridModel UngroupedModel)
         {
-			TreeGridModel GroupedModel = new TreeGridModel();
+			IEnumerable<TimingDataViewModel> TimingDataModel = UngroupedModel.Cast<TimingDataViewModel>();
+			double OverrideDuration = TimingDataModel.ElementAt(0).ParentDurationOverride ?? TimingDataModel.Sum(d => d.InclusiveDuration);
 			List<TimingDataViewModel> ClonedChildren = new List<TimingDataViewModel>();
-            foreach (TimingDataViewModel Child in UngroupedModel.Cast<TimingDataViewModel>())
+            foreach (TimingDataViewModel Child in TimingDataModel)
             {
 				TimingDataViewModel ClonedChild = Child.Clone();
                 if (ClonedChild.HasChildren)
@@ -262,12 +283,11 @@ namespace Timing_Data_Investigator.Controls
             }
 
 			IEnumerable<TimingDataViewModel> GroupedClonedChildren = GroupChildren(ClonedChildren);
-            foreach (TimingDataViewModel GroupedChild in GroupedClonedChildren)
-            {
-                GroupedModel.Add(GroupedChild);
-            }
-
-            return GroupedModel;
+			foreach (TimingDataViewModel Child in GroupedClonedChildren)
+			{
+				Child.ParentDurationOverride = OverrideDuration;
+			}
+			return new TreeGridModel(GroupedClonedChildren);
         }
 
         private IEnumerable<TimingDataViewModel> GroupChildren(IEnumerable<TimingDataViewModel> Children)
@@ -390,5 +410,24 @@ namespace Timing_Data_Investigator.Controls
 				UpdateGridModel(GroupFunctionTemplates, FunctionsGrid, FunctionsModel, GroupedFunctionsModel);
 			}
 		}
+
+		private void DataGridRow_MouseEnter(object sender, System.Windows.Input.MouseEventArgs e)
+		{
+			DataGridRow Row = sender as DataGridRow;
+			Row.Background = SystemColors.InactiveSelectionHighlightBrush;
+		}
+
+		private void DataGridRow_MouseLeave(object sender, System.Windows.Input.MouseEventArgs e)
+		{
+			DataGridRow Row = sender as DataGridRow;
+			Row.Background = Brushes.White;
+		}
+
+		private void DataGridRow_MouseDoubleClick(object sender, System.Windows.Input.MouseButtonEventArgs e)
+		{
+			DataGridRow Row = sender as DataGridRow;
+			TimingDataViewModel File = Row.DataContext as TimingDataViewModel;
+			File.OpenCommand.Execute(null);
+		}
 	}
 }
diff --git a/Engine/Source/Programs/Windows/CompileTimeAnalyzer/CompileTimeAnalyzer/MainWindow.xaml b/Engine/Source/Programs/Windows/CompileTimeAnalyzer/CompileTimeAnalyzer/MainWindow.xaml
--- a/Engine/Source/Programs/Windows/CompileTimeAnalyzer/CompileTimeAnalyzer/MainWindow.xaml
+++ b/Engine/Source/Programs/Windows/CompileTimeAnalyzer/CompileTimeAnalyzer/MainWindow.xaml
@@ -40,15 +40,17 @@
         <DockPanel Grid.Row="0">
             <Menu DockPanel.Dock="Top">
                 <MenuItem Header="_File">
-                    <MenuItem Header="_Open" Click="OpenFile_Click" />
-                    <Separator />
+                    <MenuItem Header="_Open Compile Time File" Click="OpenFile_Click" />
+					<Separator />
+					<MenuItem x:Name="RecentFiles" Header="Open Recent _File" IsEnabled="False" />
+					<Separator />
                     <MenuItem Header="E_xit" Click="Exit_Click" />
                 </MenuItem>
             </Menu>
             <TabControl x:Name="OpenedFiles">
                 <TabControl.ItemTemplate>
                     <DataTemplate>
-                        <StackPanel Orientation="Horizontal">
+                        <StackPanel Orientation="Horizontal" MouseDown="StackPanel_MouseDown">
 							<TextBlock Text="{Binding ShortName}" VerticalAlignment="Center" />
 							<Button Click="RemoveTab_Click" Style="{StaticResource TabCloseButton}" />
 						</StackPanel>
diff --git a/Engine/Source/Programs/Windows/CompileTimeAnalyzer/CompileTimeAnalyzer/MainWindow.xaml.cs b/Engine/Source/Programs/Windows/CompileTimeAnalyzer/CompileTimeAnalyzer/MainWindow.xaml.cs
--- a/Engine/Source/Programs/Windows/CompileTimeAnalyzer/CompileTimeAnalyzer/MainWindow.xaml.cs
+++ b/Engine/Source/Programs/Windows/CompileTimeAnalyzer/CompileTimeAnalyzer/MainWindow.xaml.cs
@@ -71,18 +71,33 @@ namespace Timing_Data_Investigator
         {
 			Button RemoveTabButton = (Button)e.Source;
 			TimingDataViewModel TabToRemove = (TimingDataViewModel)RemoveTabButton.DataContext;
-            int IndexToSelect = OpenedFiles.Items.IndexOf(TabToRemove) - 1;
-            OpenedFiles.Items.Remove(TabToRemove);
-            if (IndexToSelect == 0 && OpenedFiles.Items.Count == 1)
-            {
-                NoOpenTabsTab.Visibility = Visibility.Visible;
-            }
-            else
-            {
-                IndexToSelect = 1;
-            }
-
-            OpenedFiles.SelectedIndex = IndexToSelect;
+			CloseTab(TabToRemove);
         }
-    }
+
+		private void StackPanel_MouseDown(object sender, System.Windows.Input.MouseButtonEventArgs e)
+		{
+			if (e.ChangedButton == System.Windows.Input.MouseButton.Middle)
+			{
+				StackPanel TabStackPanel = (StackPanel)sender;
+				TimingDataViewModel TabToRemove = (TimingDataViewModel)TabStackPanel.DataContext;
+				CloseTab(TabToRemove);
+			}
+		}
+
+		private void CloseTab(TimingDataViewModel TabToRemove)
+		{
+			int IndexToSelect = OpenedFiles.Items.IndexOf(TabToRemove) - 1;
+			OpenedFiles.Items.Remove(TabToRemove);
+			if (IndexToSelect == 0 && OpenedFiles.Items.Count == 1)
+			{
+				NoOpenTabsTab.Visibility = Visibility.Visible;
+			}
+			else
+			{
+				IndexToSelect = 1;
+			}
+
+			OpenedFiles.SelectedIndex = IndexToSelect;
+		}
+	}
 }
diff --git a/Engine/Source/Programs/Windows/CompileTimeAnalyzer/CompileTimeAnalyzer/Models/TimingDataViewModel.cs b/Engine/Source/Programs/Windows/CompileTimeAnalyzer/CompileTimeAnalyzer/Models/TimingDataViewModel.cs
--- a/Engine/Source/Programs/Windows/CompileTimeAnalyzer/CompileTimeAnalyzer/Models/TimingDataViewModel.cs
+++ b/Engine/Source/Programs/Windows/CompileTimeAnalyzer/CompileTimeAnalyzer/Models/TimingDataViewModel.cs
@@ -35,15 +35,32 @@ namespace Timing_Data_Investigator.Models
 
 		public string Name { get; set; }
 		public TimingDataType Type { get; set; }
+		public int Count { get; set; } = 1;
+		public double? ParentDurationOverride { get; set; }
 		public double ExclusiveDuration { get; set; }
+		public double ExclusivePercent => GetPercentOfParent(ExclusiveDuration);
 		public double InclusiveDuration => ExclusiveDuration + Children.Cast<TimingDataViewModel>().Sum(c => c.InclusiveDuration);
-
+		public double InclusivePercent => GetPercentOfParent(InclusiveDuration);
 		public string ShortName => !string.IsNullOrWhiteSpace(Name) ? Path.GetFileName(Name) : null;
-
 		public ICommand OpenCommand { get; set; }
-
 		public int SortedIndex { get; set; }
 
+		private double GetPercentOfParent(double Duration)
+		{
+			if (Parent == null)
+			{
+				if (ParentDurationOverride.HasValue)
+				{
+					return (100 * Duration) / ParentDurationOverride.Value;
+				}
+
+				return 100.0;
+			}
+
+			TimingDataViewModel ParentViewModel = (TimingDataViewModel)Parent;
+			return (100 * Duration) / ParentViewModel.InclusiveDuration;
+		}
+
 		public void UpdateSortIndex(string SortMember, ListSortDirection SortDirection)
 		{
 			int Index = 0;
@@ -136,6 +153,7 @@ namespace Timing_Data_Investigator.Models
 			{
 				Name = TimingData.Name,
 				Type = TimingData.Type,
+				Count = TimingData.Count,
 				ExclusiveDuration = TimingData.ExclusiveDuration,
 				HasChildren = TimingData.Children.Any(),
 			};
@@ -216,6 +234,7 @@ namespace Timing_Data_Investigator.Models
 			{
 				Name = Name,
 				Type = Type,
+				ParentDurationOverride = ParentDurationOverride,
 				ExclusiveDuration = ExclusiveDuration,
 				HasChildren = Children.Any(),
 			};
diff --git a/Engine/Source/Programs/Windows/CompileTimeAnalyzer/CompileTimeAnalyzer/TreeGridFlatModel.cs b/Engine/Source/Programs/Windows/CompileTimeAnalyzer/CompileTimeAnalyzer/TreeGridFlatModel.cs
--- a/Engine/Source/Programs/Windows/CompileTimeAnalyzer/CompileTimeAnalyzer/TreeGridFlatModel.cs
+++ b/Engine/Source/Programs/Windows/CompileTimeAnalyzer/CompileTimeAnalyzer/TreeGridFlatModel.cs
@@ -19,10 +19,19 @@ namespace Timing_Data_Investigator
 		private PropertyInfo CurrentSortProperty;
 		private ListSortDirection? CurrentSortDirection;
 
-		public TreeGridFlatModel()
+		public TreeGridFlatModel() : this(null)
+		{
+		}
+
+		public TreeGridFlatModel(IEnumerable<TreeGridElement> Elements)
 		{
 			// Initialize the model
 			Keys = new HashSet<TreeGridElement>();
+			foreach (TreeGridElement Element in Elements)
+			{
+				Items.Add(Element);
+				Keys.Add(Element);
+			}
 		}
 
 		public void Sort(PropertyInfo SortProperty, ListSortDirection? SortDirection)
diff --git a/Engine/Source/Programs/Windows/CompileTimeAnalyzer/CompileTimeAnalyzer/TreeGridModel.cs b/Engine/Source/Programs/Windows/CompileTimeAnalyzer/CompileTimeAnalyzer/TreeGridModel.cs
--- a/Engine/Source/Programs/Windows/CompileTimeAnalyzer/CompileTimeAnalyzer/TreeGridModel.cs
+++ b/Engine/Source/Programs/Windows/CompileTimeAnalyzer/CompileTimeAnalyzer/TreeGridModel.cs
@@ -11,11 +11,23 @@ namespace Timing_Data_Investigator
 	{
 		private List<TreeGridElement> ItemCache;
 
-		public TreeGridModel()
+		public TreeGridModel() : this(null)
+		{
+		}
+
+		public TreeGridModel(IEnumerable<TreeGridElement> Elements)
 		{
 			// Initialize the model
 			ItemCache = new List<TreeGridElement>();
-			FlatModel = new TreeGridFlatModel();
+			FlatModel = new TreeGridFlatModel(Elements);
+			if (Elements != null)
+			{
+				foreach (TreeGridElement Element in Elements)
+				{
+					Items.Add(Element);
+					OnRootAdded(Element, false);
+				}
+			}
 		}
 
 		public TreeGridFlatModel FlatModel { get; private set; }
@@ -28,7 +40,7 @@ namespace Timing_Data_Investigator
 				case NotifyCollectionChangedAction.Add:
 
 					// Process added item
-					OnRootAdded(args.NewItems[0]);
+					OnRootAdded(args.NewItems[0], true);
 					break;
 			}
 		}
@@ -106,7 +118,7 @@ namespace Timing_Data_Investigator
 		{
 		}
 
-		private void OnRootAdded(object Item)
+		private void OnRootAdded(object Item, bool InsertIntoFlatModel)
 		{
 			// Verify the root item
 			TreeGridElement Root = TreeGridElement.VerifyItem(Item);
@@ -114,14 +126,17 @@ namespace Timing_Data_Investigator
 			// Set the model for the root
 			Root.SetModel(this);
 
-			// Find the index for insertion into the flat model
-			int index = FindFlatInsertionIndex(Root);
+			if (InsertIntoFlatModel)
+			{
+				// Find the index for insertion into the flat model
+				int index = FindFlatInsertionIndex(Root);
 
-			// Insert the root into the flat model
-			FlatModel.PrivateInsert(index, Root);
+				// Insert the root into the flat model
+				FlatModel.PrivateInsert(index, Root);
 
-			// Expand the root
-			Expand(Root);
+				// Expand the root
+				Expand(Root);
+			}
 		}
 
 		private void CacheFlatChildren(TreeGridElement item)
