Suppress a warning in CompileTimeAnalyzer.

#rnx

diff --git a/Engine/Source/Programs/Windows/CompileTimeAnalyzer/CompileTimeAnalyzer/Commands/OpenTimingDataCommand.cs b/Engine/Source/Programs/Windows/CompileTimeAnalyzer/CompileTimeAnalyzer/Commands/OpenTimingDataCommand.cs
--- a/Engine/Source/Programs/Windows/CompileTimeAnalyzer/CompileTimeAnalyzer/Commands/OpenTimingDataCommand.cs
+++ b/Engine/Source/Programs/Windows/CompileTimeAnalyzer/CompileTimeAnalyzer/Commands/OpenTimingDataCommand.cs
@@ -10,7 +10,9 @@ namespace Timing_Data_Investigator.Commands
     {
         private TimingDataViewModel ViewModelToOpen;
 
+		#pragma warning disable CS0067
 		public event EventHandler CanExecuteChanged;
+		#pragma warning restore CS0067
 
         public OpenTimingDataCommand(TimingDataViewModel TimingData)
         {
