Cosmetic changes to compile time analyzer

* Add a minimum width for columns on the timing data grid
* Lots of margins around controls
* Add a UE4 icon
* Show the open file dialog as soon as the program opens

diff --git a/Engine/Source/Programs/Windows/CompileTimeAnalyzer/CompileTimeAnalyzer/CompileTimeAnalyzer.csproj b/Engine/Source/Programs/Windows/CompileTimeAnalyzer/CompileTimeAnalyzer/CompileTimeAnalyzer.csproj
--- a/Engine/Source/Programs/Windows/CompileTimeAnalyzer/CompileTimeAnalyzer/CompileTimeAnalyzer.csproj
+++ b/Engine/Source/Programs/Windows/CompileTimeAnalyzer/CompileTimeAnalyzer/CompileTimeAnalyzer.csproj
@@ -35,6 +35,9 @@
     <ErrorReport>prompt</ErrorReport>
     <WarningLevel>4</WarningLevel>
   </PropertyGroup>
+  <PropertyGroup>
+    <ApplicationIcon>Default.ico</ApplicationIcon>
+  </PropertyGroup>
   <ItemGroup>
     <Reference Include="System" />
     <Reference Include="System.Data" />
@@ -126,5 +129,8 @@
       <Name>UnrealBuildTool</Name>
     </ProjectReference>
   </ItemGroup>
+  <ItemGroup>
+    <Resource Include="Default.ico" />
+  </ItemGroup>
   <Import Project="$(MSBuildToolsPath)\Microsoft.CSharp.targets" />
 </Project>
\ No newline at end of file
diff --git a/Engine/Source/Programs/Windows/CompileTimeAnalyzer/CompileTimeAnalyzer/Controls/TimingDataGrid.xaml b/Engine/Source/Programs/Windows/CompileTimeAnalyzer/CompileTimeAnalyzer/Controls/TimingDataGrid.xaml
--- a/Engine/Source/Programs/Windows/CompileTimeAnalyzer/CompileTimeAnalyzer/Controls/TimingDataGrid.xaml
+++ b/Engine/Source/Programs/Windows/CompileTimeAnalyzer/CompileTimeAnalyzer/Controls/TimingDataGrid.xaml
@@ -89,7 +89,7 @@
 									</Style>
 								</Button.Style>
 							</Button>
-							<CheckBox Grid.Column="1" IsChecked="{Binding IsExpanded, Mode=TwoWay, UpdateSourceTrigger=PropertyChanged}" Visibility="{Binding HasChildren, Converter={StaticResource BoolToVisibilityConverter}}" Style="{StaticResource Expander}" />
+							<CheckBox Grid.Column="1" IsChecked="{Binding IsExpanded, Mode=TwoWay, UpdateSourceTrigger=PropertyChanged}" Visibility="{Binding HasChildren, Converter={StaticResource BoolToVisibilityConverter}}" Style="{StaticResource Expander}" Margin="2"/>
 							<Line Grid.Column="1" X1="8" X2="8" Y2="{Binding ActualHeight, RelativeSource={RelativeSource Self}}" Stroke="#888" StrokeDashArray="2 2" StrokeDashOffset="1">
 								<Line.Visibility>
 									<MultiBinding Converter="{StaticResource MultiBoolToVisibilityConverter}" ConverterParameter="VisibleOnFalse">
@@ -98,15 +98,31 @@
 									</MultiBinding>
 								</Line.Visibility>
 							</Line>
-							<TextBlock Grid.Column="2" Text="{Binding Name}"/>
+							<TextBlock Grid.Column="2" Text="{Binding Name}" Margin="2"/>
 						</Grid>
 					</DataTemplate>
 				</DataGridTemplateColumn.CellTemplate>
 			</DataGridTemplateColumn>
 
-			<DataGridTextColumn x:Name="CountColumn" Header="Count" Binding="{Binding Count}" SortMemberPath="Count" />
+			<DataGridTextColumn x:Name="CountColumn" Header="Count" Binding="{Binding Count}" SortMemberPath="Count" Width="100">
+				<DataGridTextColumn.HeaderStyle>
+					<Style TargetType="DataGridColumnHeader">
+						<Setter Property="HorizontalContentAlignment" Value="Center" />
+					</Style>
+				</DataGridTextColumn.HeaderStyle>
+				<DataGridTextColumn.ElementStyle>
+					<Style TargetType="TextBlock">
+						<Setter Property="HorizontalAlignment" Value="Right" />
+					</Style>
+				</DataGridTextColumn.ElementStyle>
+			</DataGridTextColumn>
 
-			<DataGridTemplateColumn Header="Duration" CanUserSort="True" SortMemberPath="InclusiveDuration">
+			<DataGridTemplateColumn Header="Time (Inclusive)" CanUserSort="True" SortMemberPath="InclusiveDuration" Width="100">
+				<DataGridTemplateColumn.HeaderStyle>
+					<Style TargetType="DataGridColumnHeader">
+						<Setter Property="HorizontalContentAlignment" Value="Center" />
+					</Style>
+				</DataGridTemplateColumn.HeaderStyle>
 				<DataGridTemplateColumn.CellTemplate>
 					<DataTemplate>
 						<Grid>
@@ -114,7 +130,7 @@
 										 Margin="2,1,2,1" Foreground="CornflowerBlue" Background="Transparent" BorderThickness="0" />
 							<TextBlock TextAlignment="Right" Margin="4,0,4,0">
 								<TextBlock.Text>
-									<MultiBinding StringFormat="{}{0:0.###}s ({1:0.##}%)">
+									<MultiBinding StringFormat="{}{0:0.000}s ({1:0.00}%)">
 										<Binding Path="InclusiveDuration" />
 										<Binding Path="InclusivePercent" />
 									</MultiBinding>
@@ -125,7 +141,12 @@
 				</DataGridTemplateColumn.CellTemplate>
 			</DataGridTemplateColumn>
 
-			<DataGridTemplateColumn x:Name="ExclusiveDurationColumn" Header="Duration (Exclusive)" CanUserSort="True" SortMemberPath="ExclusiveDuration">
+			<DataGridTemplateColumn x:Name="ExclusiveDurationColumn" Header="Time (Exclusive)" CanUserSort="True" SortMemberPath="ExclusiveDuration" Width="100">
+				<DataGridTemplateColumn.HeaderStyle>
+					<Style TargetType="DataGridColumnHeader">
+						<Setter Property="HorizontalContentAlignment" Value="Center" />
+					</Style>
+				</DataGridTemplateColumn.HeaderStyle>
 				<DataGridTemplateColumn.CellTemplate>
 					<DataTemplate>
 						<Grid>
@@ -133,7 +154,7 @@
 										 Margin="2,1,2,1" Foreground="CornflowerBlue" Background="Transparent" BorderThickness="0" />
 							<TextBlock TextAlignment="Right" Margin="4,0,4,0">
 								<TextBlock.Text>
-									<MultiBinding StringFormat="{}{0:0.###}s ({1:0.##}%)">
+									<MultiBinding StringFormat="{}{0:0.000}s ({1:0.00}%)">
 										<Binding Path="ExclusiveDuration" />
 										<Binding Path="ExclusivePercent" />
 									</MultiBinding>
diff --git a/Engine/Source/Programs/Windows/CompileTimeAnalyzer/CompileTimeAnalyzer/Controls/TimingDataTab.xaml b/Engine/Source/Programs/Windows/CompileTimeAnalyzer/CompileTimeAnalyzer/Controls/TimingDataTab.xaml
--- a/Engine/Source/Programs/Windows/CompileTimeAnalyzer/CompileTimeAnalyzer/Controls/TimingDataTab.xaml
+++ b/Engine/Source/Programs/Windows/CompileTimeAnalyzer/CompileTimeAnalyzer/Controls/TimingDataTab.xaml
@@ -11,48 +11,52 @@
         <ContextMenu  x:Key="RowMenu" DataContext="{Binding PlacementTarget.DataContext, RelativeSource={RelativeSource Self}}">
             <MenuItem Header="{Binding Name, StringFormat=Open {0} in New Tab}" Command="{Binding OpenCommand}" />
         </ContextMenu>
-		<Style x:Key="FilesRowStyle" TargetType="{x:Type DataGridRow}">
-			<EventSetter Event="MouseDoubleClick" Handler="DataGridRow_MouseDoubleClick" />
-			<EventSetter Event="MouseEnter" Handler="DataGridRow_MouseEnter" />
-			<EventSetter Event="MouseLeave" Handler="DataGridRow_MouseLeave" />
+        <Style x:Key="FilesRowStyle" TargetType="{x:Type DataGridRow}">
+            <EventSetter Event="MouseDoubleClick" Handler="DataGridRow_MouseDoubleClick" />
+            <EventSetter Event="MouseEnter" Handler="DataGridRow_MouseEnter" />
+            <EventSetter Event="MouseLeave" Handler="DataGridRow_MouseLeave" />
         </Style>
     </UserControl.Resources>
 
     <TabControl x:Name="SummaryTabs">
-		<TabItem x:Name="FilesTab" Header="Files" Visibility="Collapsed">
-            <local:TimingDataGrid x:Name="FilesGrid" DockPanel.Dock="Top" />
+        <TabItem x:Name="FilesTab" Header="Files" Visibility="Collapsed">
+            <local:TimingDataGrid x:Name="FilesGrid" DockPanel.Dock="Top" Margin="4" />
         </TabItem>
         <TabItem x:Name="IncludesTab" Header="Includes">
-			<Grid>
-				<Grid.RowDefinitions>
-					<RowDefinition Height="Auto" />
-					<RowDefinition Height="*" />
-				</Grid.RowDefinitions>
-				
-				<CheckBox x:Name="FlattenIncludes" Checked="FlattenIncludes_Checked" Unchecked="FlattenIncludes_Checked" VerticalAlignment="Center" Margin="4" Content="Flatten Includes" />
-				<local:TimingDataGrid Grid.Row="1" x:Name="IncludesGrid" DockPanel.Dock="Top" />
-			</Grid>
-		</TabItem>
+            <Grid Margin="6">
+                <Grid.RowDefinitions>
+                    <RowDefinition Height="Auto" />
+                    <RowDefinition Height="*" />
+                </Grid.RowDefinitions>
+
+                <CheckBox x:Name="FlattenIncludes" Checked="FlattenIncludes_Checked" Unchecked="FlattenIncludes_Checked" VerticalAlignment="Center" Margin="0" Content="Flatten Includes" />
+                <local:TimingDataGrid Grid.Row="1" x:Name="IncludesGrid" DockPanel.Dock="Top" Margin="0,6,0,0" />
+            </Grid>
+        </TabItem>
         <TabItem x:Name="ClassesTab" Header="Classes">
-            <Grid>
+            <Grid Margin="6">
                 <Grid.RowDefinitions>
                     <RowDefinition Height="Auto" />
                     <RowDefinition Height="*" />
                 </Grid.RowDefinitions>
-				
-				<CheckBox x:Name="GroupClassTemplates" Checked="GroupTemplates_Checked" Unchecked="GroupTemplates_Checked" VerticalAlignment="Center" Margin="4" Content="Group Template Classes" />
-                <local:TimingDataGrid Grid.Row="1" x:Name="ClassesGrid" />
+
+                <CheckBox x:Name="GroupClassTemplates" Checked="GroupTemplates_Checked" Unchecked="GroupTemplates_Checked" VerticalAlignment="Center" Margin="0" Content="Group Template Classes" />
+                <local:TimingDataGrid Grid.Row="1" x:Name="ClassesGrid" Margin="0,6,0,0" />
             </Grid>
         </TabItem>
         <TabItem x:Name="FunctionsTab" Header="Functions">
-            <Grid>
+            <Grid Margin="6">
+                <Grid.ColumnDefinitions>
+                    <ColumnDefinition Width="40*"/>
+                    <ColumnDefinition Width="357*"/>
+                </Grid.ColumnDefinitions>
                 <Grid.RowDefinitions>
                     <RowDefinition Height="Auto" />
                     <RowDefinition Height="*" />
                 </Grid.RowDefinitions>
-				
-				<CheckBox x:Name="GroupFunctionTemplates" Checked="GroupTemplates_Checked" Unchecked="GroupTemplates_Checked" VerticalAlignment="Center" Margin="4" Content="Group Template Functions" />
-                <local:TimingDataGrid Grid.Row="1" x:Name="FunctionsGrid" />
+
+                <CheckBox x:Name="GroupFunctionTemplates" Checked="GroupTemplates_Checked" Unchecked="GroupTemplates_Checked" VerticalAlignment="Center" Margin="0" Content="Group Template Functions" Grid.ColumnSpan="2" />
+                <local:TimingDataGrid Grid.Row="1" x:Name="FunctionsGrid" Margin="0,6,0,0" Grid.ColumnSpan="2" />
             </Grid>
         </TabItem>
     </TabControl>
diff --git a/Engine/Source/Programs/Windows/CompileTimeAnalyzer/CompileTimeAnalyzer/MainWindow.xaml b/Engine/Source/Programs/Windows/CompileTimeAnalyzer/CompileTimeAnalyzer/MainWindow.xaml
--- a/Engine/Source/Programs/Windows/CompileTimeAnalyzer/CompileTimeAnalyzer/MainWindow.xaml
+++ b/Engine/Source/Programs/Windows/CompileTimeAnalyzer/CompileTimeAnalyzer/MainWindow.xaml
@@ -6,7 +6,7 @@
         xmlns:controls="clr-namespace:Timing_Data_Investigator.Controls"
         xmlns:tdi="clr-namespace:Timing_Data_Investigator"
         mc:Ignorable="d"
-        Title="CompileTimeAnalyzer" Width="800" Height="600">
+        Title="CompileTimeAnalyzer" Width="1024" Height="768" WindowStartupLocation="CenterScreen">
 
     <Window.Resources>
 		<tdi:NullToVisibilityConverter x:Key="NullToVisibilityConverter" />
@@ -38,7 +38,7 @@
             <RowDefinition Height="Auto" />
         </Grid.RowDefinitions>
         <DockPanel Grid.Row="0">
-            <Menu DockPanel.Dock="Top">
+            <Menu DockPanel.Dock="Top" Background="{DynamicResource {x:Static SystemColors.WindowBrushKey}}">
                 <MenuItem Header="_File">
                     <MenuItem Header="_Open Compile Time File" Click="OpenFile_Click" />
 					<Separator />
@@ -47,7 +47,7 @@
                     <MenuItem Header="E_xit" Click="Exit_Click" />
                 </MenuItem>
             </Menu>
-            <TabControl x:Name="OpenedFiles">
+            <TabControl x:Name="OpenedFiles" Margin="6">
                 <TabControl.ItemTemplate>
                     <DataTemplate>
                         <StackPanel Orientation="Horizontal" MouseDown="StackPanel_MouseDown">
@@ -58,7 +58,7 @@
                 </TabControl.ItemTemplate>
                 <TabControl.ContentTemplate>
                     <DataTemplate>
-                        <controls:TimingDataTab />
+                        <controls:TimingDataTab Margin="4" />
                     </DataTemplate>
                 </TabControl.ContentTemplate>
                 <TabItem x:Name="NoOpenTabsTab" Header="Open a File to Investigate..." Opacity="0.4"/>
diff --git a/Engine/Source/Programs/Windows/CompileTimeAnalyzer/CompileTimeAnalyzer/MainWindow.xaml.cs b/Engine/Source/Programs/Windows/CompileTimeAnalyzer/CompileTimeAnalyzer/MainWindow.xaml.cs
--- a/Engine/Source/Programs/Windows/CompileTimeAnalyzer/CompileTimeAnalyzer/MainWindow.xaml.cs
+++ b/Engine/Source/Programs/Windows/CompileTimeAnalyzer/CompileTimeAnalyzer/MainWindow.xaml.cs
@@ -19,9 +19,16 @@ namespace Timing_Data_Investigator
         public MainWindow()
         {
             InitializeComponent();
+
+			Loaded += MainWindow_Loaded;
         }
 
-        private void LoadTimingFile(string FilePath)
+		private void MainWindow_Loaded(object sender, RoutedEventArgs e)
+		{
+			ShowOpenFileDialog();
+		}
+
+		private void LoadTimingFile(string FilePath)
         {
 			TimingDataViewModel NewTimingData = TimingDataViewModel.FromBinaryFile(FileReference.FromString(FilePath));
 
@@ -54,9 +61,14 @@ namespace Timing_Data_Investigator
 
         private void OpenFile_Click(object sender, RoutedEventArgs e)
         {
+			ShowOpenFileDialog();
+		}
+
+		private void ShowOpenFileDialog()
+		{
 			OpenFileDialog OpenFileDialog = new OpenFileDialog();
             OpenFileDialog.Filter = "Timing Files (*.timing.bin)|*.timing.bin|All Files (*.*)|*.*";
-            if (OpenFileDialog.ShowDialog() == true)
+            if (OpenFileDialog.ShowDialog(this) == true)
             {
                 LoadTimingFile(OpenFileDialog.FileName);
             }
