Inherit handles when launching game process, allows to continue to read stdout/stdin as the child processes will use the same handles for those streams.

From UDN 524523

diff --git a/Engine/Source/Programs/Windows/BootstrapPackagedGame/Private/BootstrapPackagedGame.cpp b/Engine/Source/Programs/Windows/BootstrapPackagedGame/Private/BootstrapPackagedGame.cpp
--- a/Engine/Source/Programs/Windows/BootstrapPackagedGame/Private/BootstrapPackagedGame.cpp
+++ b/Engine/Source/Programs/Windows/BootstrapPackagedGame/Private/BootstrapPackagedGame.cpp
@@ -120,7 +120,7 @@ int SpawnTarget(WCHAR* CmdLine)
 	PROCESS_INFORMATION ProcessInfo;
 	ZeroMemory(&ProcessInfo, sizeof(ProcessInfo));
 
-	if(!CreateProcess(NULL, CmdLine, NULL, NULL, FALSE, 0, NULL, NULL, &StartupInfo, &ProcessInfo))
+	if(!CreateProcess(NULL, CmdLine, NULL, NULL, TRUE, 0, NULL, NULL, &StartupInfo, &ProcessInfo))
 	{
 		DWORD ErrorCode = GetLastError();
 
