Refactoring how timing data is loaded to improve load times.

#rnx

diff --git a/Engine/Source/Programs/Windows/CompileTimeAnalyzer/CompileTimeAnalyzer/Controls/TimingDataTab.xaml.cs b/Engine/Source/Programs/Windows/CompileTimeAnalyzer/CompileTimeAnalyzer/Controls/TimingDataTab.xaml.cs
--- a/Engine/Source/Programs/Windows/CompileTimeAnalyzer/CompileTimeAnalyzer/Controls/TimingDataTab.xaml.cs
+++ b/Engine/Source/Programs/Windows/CompileTimeAnalyzer/CompileTimeAnalyzer/Controls/TimingDataTab.xaml.cs
@@ -17,87 +17,87 @@ namespace Timing_Data_Investigator.Controls
 	/// Interaction logic for TimingDataTab.xaml
 	/// </summary>
 	public partial class TimingDataTab : UserControl
-    {
-        private TreeGridModel FilesModel;
-        private TreeGridModel IncludesModel;
+	{
+		private TreeGridModel FilesModel;
+		private TreeGridModel IncludesModel;
 		private TreeGridModel FlattenedIncludesModel;
 		private TreeGridModel ClassesModel;
-        private TreeGridModel GroupedClassesModel;
-        private TreeGridModel FunctionsModel;
-        private TreeGridModel GroupedFunctionsModel;
+		private TreeGridModel GroupedClassesModel;
+		private TreeGridModel FunctionsModel;
+		private TreeGridModel GroupedFunctionsModel;
 
 		private Dictionary<TimingDataViewModel, TabState> TabStates = new Dictionary<TimingDataViewModel, TabState>();
 
-        private class TabState
-        {
-            public Dictionary<TabItem, DataGridColumn> SortColumns { get; } = new Dictionary<TabItem, DataGridColumn>();
-            public Dictionary<TabItem, ListSortDirection?> SortDirections { get; } = new Dictionary<TabItem, ListSortDirection?>();
+		private class TabState
+		{
+			public Dictionary<TabItem, DataGridColumn> SortColumns { get; } = new Dictionary<TabItem, DataGridColumn>();
+			public Dictionary<TabItem, ListSortDirection?> SortDirections { get; } = new Dictionary<TabItem, ListSortDirection?>();
 
 			public bool? FlattenIncludes { get; set; }
 			public bool? GroupClasses { get; set; }
-            public bool? GroupFunctions { get; set; }
-
-            public void SetTabState(TabItem Tab, DataGrid TabDataGrid)
-            {
-                DataGridColumn SortedColumn = TabDataGrid.Columns.FirstOrDefault(c => c.SortDirection != null);
-                SortColumns[Tab] = SortedColumn;
-                SortDirections[Tab] = SortedColumn?.SortDirection;
-            }
-        }
-
-        public TimingDataTab()
-        {
-            InitializeComponent();
-            DataContextChanged += TimingDataTab_DataContextChanged;
-            FilesGrid.Grid.RowStyle = FindResource("FilesRowStyle") as Style;
+			public bool? GroupFunctions { get; set; }
+
+			public void SetTabState(TabItem Tab, DataGrid TabDataGrid)
+			{
+				DataGridColumn SortedColumn = TabDataGrid.Columns.FirstOrDefault(c => c.SortDirection != null);
+				SortColumns[Tab] = SortedColumn;
+				SortDirections[Tab] = SortedColumn?.SortDirection;
+			}
 		}
 
-        private void TimingDataTab_DataContextChanged(object sender, DependencyPropertyChangedEventArgs e)
-        {
-            if (e.NewValue == null)
-            {
-                FilesModel?.Clear();
-                IncludesModel?.Clear();
-                ClassesModel?.Clear();
-                FunctionsModel?.Clear();
-                GroupedClassesModel?.Clear();
-                GroupedFunctionsModel?.Clear();
-                IncludesGrid.DataContext = null;
-                ClassesGrid.DataContext = null;
-                FunctionsGrid.DataContext = null;
-                return;
-            }
-
-            // Save off current relevant state for the old tab if needed.
-            if (e.OldValue != null)
-            {
+		public TimingDataTab()
+		{
+			InitializeComponent();
+			DataContextChanged += TimingDataTab_DataContextChanged;
+			FilesGrid.Grid.RowStyle = FindResource("FilesRowStyle") as Style;
+		}
+
+		private void TimingDataTab_DataContextChanged(object sender, DependencyPropertyChangedEventArgs e)
+		{
+			if (e.NewValue == null)
+			{
+				FilesModel?.Clear();
+				IncludesModel?.Clear();
+				ClassesModel?.Clear();
+				FunctionsModel?.Clear();
+				GroupedClassesModel?.Clear();
+				GroupedFunctionsModel?.Clear();
+				IncludesGrid.DataContext = null;
+				ClassesGrid.DataContext = null;
+				FunctionsGrid.DataContext = null;
+				return;
+			}
+
+			// Save off current relevant state for the old tab if needed.
+			if (e.OldValue != null)
+			{
 				TimingDataViewModel OldTimingData = e.OldValue as TimingDataViewModel;
-                if (!TabStates.ContainsKey(OldTimingData))
-                {
-                    TabStates.Add(OldTimingData, new TabState());
-                }
-
-                TabState TabState = TabStates[OldTimingData];
-                TabState.SetTabState(FilesTab, FilesGrid.Grid);
-                TabState.SetTabState(IncludesTab, IncludesGrid.Grid);
-                TabState.SetTabState(ClassesTab, ClassesGrid.Grid);
-                TabState.SetTabState(FunctionsTab, FunctionsGrid.Grid);
+				if (!TabStates.ContainsKey(OldTimingData))
+				{
+					TabStates.Add(OldTimingData, new TabState());
+				}
+
+				TabState TabState = TabStates[OldTimingData];
+				TabState.SetTabState(FilesTab, FilesGrid.Grid);
+				TabState.SetTabState(IncludesTab, IncludesGrid.Grid);
+				TabState.SetTabState(ClassesTab, ClassesGrid.Grid);
+				TabState.SetTabState(FunctionsTab, FunctionsGrid.Grid);
 				TabState.FlattenIncludes = FlattenIncludes.IsChecked;
-                TabState.GroupClasses = GroupClassTemplates.IsChecked;
-                TabState.GroupFunctions = GroupFunctionTemplates.IsChecked;
-            }
+				TabState.GroupClasses = GroupClassTemplates.IsChecked;
+				TabState.GroupFunctions = GroupFunctionTemplates.IsChecked;
+			}
 
 			TimingDataViewModel NewTimingData = e.NewValue as TimingDataViewModel;
 
-            // Determine whether we loaded a single file's data or a summary file.
-            if (NewTimingData.Type == TimingDataType.Aggregate)
-            {
+			// Determine whether we loaded a single file's data or a summary file.
+			if (NewTimingData.Type == TimingDataType.Aggregate)
+			{
 				UpdateUIForAggregate(NewTimingData);
 			}
-            else
-            {
+			else
+			{
 				UpdateUIForSingleFile(NewTimingData);
-            }
+			}
 
 			SummaryTabs.SelectedItem = FilesModel == null ? IncludesTab : FilesTab;
 
@@ -190,31 +190,31 @@ namespace Timing_Data_Investigator.Controls
 			State.SortColumns[Tab].SortDirection = State.SortDirections[Tab];
 		}
 
-        private void SortModel(DataGrid Grid)
-        {
+		private void SortModel(DataGrid Grid)
+		{
 			if (Grid.DataContext == null)
 			{
 				return;
 			}
 
-            // Find the column that has a sorting value set, if any.
-            DataGridColumn SortedColumn = Grid.Columns.FirstOrDefault(c => c.SortDirection != null);
-            if (SortedColumn == null)
-            {
-                // No sort applied.
-                return;
-            }
+			// Find the column that has a sorting value set, if any.
+			DataGridColumn SortedColumn = Grid.Columns.FirstOrDefault(c => c.SortDirection != null);
+			if (SortedColumn == null)
+			{
+				// No sort applied.
+				return;
+			}
 
 			TreeGridFlatModel Model = (TreeGridFlatModel)Grid.DataContext;
 			PropertyInfo SortProperty = typeof(TimingDataViewModel).GetProperty(SortedColumn.SortMemberPath);
 			Model.Sort(SortProperty, SortedColumn.SortDirection);
 		}
 
-        private void UpdateTabVisibility(TreeGridFlatModel TabModel, TabItem Tab, TimingDataGrid Grid)
-        {
-            Tab.Visibility = TabModel == null ? Visibility.Collapsed : Visibility.Visible;
+		private void UpdateTabVisibility(TreeGridFlatModel TabModel, TabItem Tab, TimingDataGrid Grid)
+		{
+			Tab.Visibility = TabModel == null ? Visibility.Collapsed : Visibility.Visible;
 			Grid.DataContext = TabModel;
-        }
+		}
 
 		private TreeGridModel GenerateTreeGridModel(TreeGridElement Parent)
 		{
@@ -263,24 +263,24 @@ namespace Timing_Data_Investigator.Controls
 		}
 
 		private TreeGridModel GenerateGroupedModel(TreeGridModel UngroupedModel)
-        {
+		{
 			IEnumerable<TimingDataViewModel> TimingDataModel = UngroupedModel.Cast<TimingDataViewModel>();
 			double OverrideDuration = TimingDataModel.ElementAt(0).ParentDurationOverride ?? TimingDataModel.Sum(d => d.InclusiveDuration);
 			List<TimingDataViewModel> ClonedChildren = new List<TimingDataViewModel>();
-            foreach (TimingDataViewModel Child in TimingDataModel)
-            {
+			foreach (TimingDataViewModel Child in TimingDataModel)
+			{
 				TimingDataViewModel ClonedChild = Child.Clone();
-                if (ClonedChild.HasChildren)
-                {
+				if (ClonedChild.HasChildren)
+				{
 					IEnumerable<TimingDataViewModel> GroupedChildren = GroupChildren(ClonedChild.Children.Cast<TimingDataViewModel>());
-                    ClonedChild.Children.Clear();
-                    foreach (TimingDataViewModel GroupedChild in GroupedChildren)
-                    {
-                        ClonedChild.Children.Add(GroupedChild);
-                    }
-                }
-                ClonedChildren.Add(ClonedChild);
-            }
+					ClonedChild.Children.Clear();
+					foreach (TimingDataViewModel GroupedChild in GroupedChildren)
+					{
+						ClonedChild.Children.Add(GroupedChild);
+					}
+				}
+				ClonedChildren.Add(ClonedChild);
+			}
 
 			IEnumerable<TimingDataViewModel> GroupedClonedChildren = GroupChildren(ClonedChildren);
 			foreach (TimingDataViewModel Child in GroupedClonedChildren)
@@ -288,115 +288,105 @@ namespace Timing_Data_Investigator.Controls
 				Child.ParentDurationOverride = OverrideDuration;
 			}
 			return new TreeGridModel(GroupedClonedChildren);
-        }
+		}
 
-        private IEnumerable<TimingDataViewModel> GroupChildren(IEnumerable<TimingDataViewModel> Children)
-        {
+		private IEnumerable<TimingDataViewModel> GroupChildren(IEnumerable<TimingDataViewModel> Children)
+		{
 			List<TimingDataViewModel> GroupedChildren = new List<TimingDataViewModel>();
 			Dictionary<string, List<TimingDataViewModel>> ChildGroups = new Dictionary<string, List<TimingDataViewModel>>();
-            foreach (TimingDataViewModel Child in Children)
-            {
-                if (Child.HasChildren)
-                {
+			foreach (TimingDataViewModel Child in Children)
+			{
+				if (Child.HasChildren)
+				{
 					IEnumerable<TimingDataViewModel> ChildsGroupedChildren = GroupChildren(Child.Children.Cast<TimingDataViewModel>());
-                    Child.Children.Clear();
-                    foreach (TimingDataViewModel ChildChild in ChildsGroupedChildren)
-                    {
-                        Child.Children.Add(ChildChild);
-                    }
-                }
+					Child.Children.Clear();
+					foreach (TimingDataViewModel ChildChild in ChildsGroupedChildren)
+					{
+						Child.Children.Add(ChildChild);
+					}
+				}
 
 				// See if this is a templated class. If not, add it as is.
 				Match Match = Regex.Match(Child.Name, @"^([^<]*)(?<Template><.*>)");
-                if (!Match.Success)
-                {
+				if (!Match.Success)
+				{
 					// Check to see if we've seen this name before. If so, group them together.
 					string DuplicateGroupName = $"{Child.Name} (Duplicates)";
-                    if (ChildGroups.ContainsKey(DuplicateGroupName))
-                    {
-                        ChildGroups[DuplicateGroupName].Add(Child);
-                    }
-                    else
-                    {
-						TimingDataViewModel FoundChild = GroupedChildren.FirstOrDefault(c => c.Name == Child.Name);
-                        if (FoundChild != null)
-                        {
-                            ChildGroups.Add(DuplicateGroupName, new List<TimingDataViewModel>());
-                            ChildGroups[DuplicateGroupName].Add(FoundChild);
-                            ChildGroups[DuplicateGroupName].Add(Child);
-                            GroupedChildren.Remove(FoundChild);
-                        }
-                        else
-                        {
-                            GroupedChildren.Add(Child);
-                        }
-                    }
-                }
-                else
-                {
+					if (ChildGroups.ContainsKey(DuplicateGroupName))
+					{
+						ChildGroups[DuplicateGroupName].Add(Child);
+					}
+					else
+					{
+						ChildGroups.Add(DuplicateGroupName, new List<TimingDataViewModel>());
+						ChildGroups[DuplicateGroupName].Add(Child);
+					}
+				}
+				else
+				{
 					// Generate group name from template.
 					int TemplateParamCount = Match.Groups["Template"].Value.Count(c => c == ',') + 1;
 					List<string> TemplateParamSig = new List<string>(TemplateParamCount);
-                    for (int i = 0; i < TemplateParamCount; ++i)
-                    {
-                        TemplateParamSig.Add("...");
-                    }
+					for (int i = 0; i < TemplateParamCount; ++i)
+					{
+						TemplateParamSig.Add("...");
+					}
 
 					string GroupName = Child.Name.Replace(Match.Groups["Template"].Value, $"<{string.Join(", ", TemplateParamSig)}>");
 
-                    // See if we have a group for this template already. If not, add it.
-                    if (!ChildGroups.ContainsKey(GroupName))
-                    {
-                        ChildGroups.Add(GroupName, new List<TimingDataViewModel>());
-                    }
-
-                    ChildGroups[GroupName].Add(Child);
-                }
-            }
-
-            // Add grouped children.
-            foreach (KeyValuePair<string, List<TimingDataViewModel>> Group in ChildGroups)
-            {
-                if (Group.Value.Count == 1)
-                {
-                    GroupedChildren.Add(Group.Value.First());
-                    continue;
-                }
+					// See if we have a group for this template already. If not, add it.
+					if (!ChildGroups.ContainsKey(GroupName))
+					{
+						ChildGroups.Add(GroupName, new List<TimingDataViewModel>());
+					}
+
+					ChildGroups[GroupName].Add(Child);
+				}
+			}
+
+			// Add grouped children.
+			foreach (KeyValuePair<string, List<TimingDataViewModel>> Group in ChildGroups)
+			{
+				if (Group.Value.Count == 1)
+				{
+					GroupedChildren.Add(Group.Value.First());
+					continue;
+				}
 
 				TimingDataViewModel NewViewModel = new TimingDataViewModel()
-                {
-                    Name = Group.Key,
-                    HasChildren = true,
-                };
+				{
+					Name = Group.Key,
+					HasChildren = true,
+				};
 
-                foreach (TimingDataViewModel Child in Group.Value)
-                {
-                    NewViewModel.Children.Add(Child);
-                }
+				foreach (TimingDataViewModel Child in Group.Value)
+				{
+					NewViewModel.Children.Add(Child);
+				}
 
-                GroupedChildren.Add(NewViewModel);
-            }
+				GroupedChildren.Add(NewViewModel);
+			}
 
-            return GroupedChildren;
-        }
+			return GroupedChildren;
+		}
 
-        private void UpdateGridModel(CheckBox CheckBox, TimingDataGrid Grid, TreeGridModel UncheckedModel, TreeGridModel CheckedModel)
-        {
+		private void UpdateGridModel(CheckBox CheckBox, TimingDataGrid Grid, TreeGridModel UncheckedModel, TreeGridModel CheckedModel)
+		{
 			if (CheckBox.IsChecked == true)
-            {
+			{
 				Grid.DataContext = CheckedModel?.FlatModel;
-            }
-            else
-            {
+			}
+			else
+			{
 				Grid.DataContext = UncheckedModel?.FlatModel;
-            }
-        }
+			}
+		}
 
-        private void FlattenIncludes_Checked(object sender, RoutedEventArgs e)
-        {
+		private void FlattenIncludes_Checked(object sender, RoutedEventArgs e)
+		{
 			CheckBox GroupCheckBox = sender as CheckBox;
-            UpdateGridModel(FlattenIncludes, IncludesGrid, IncludesModel, FlattenedIncludesModel);
-        }
+			UpdateGridModel(FlattenIncludes, IncludesGrid, IncludesModel, FlattenedIncludesModel);
+		}
 
 		private void GroupTemplates_Checked(object sender, RoutedEventArgs e)
 		{
