Fixing a bug in CompileTimeAnalyzer where expanded nodes wouldn't get their children sorted until the sorting method changed.

#rnx

diff --git a/Engine/Source/Programs/Windows/CompileTimeAnalyzer/CompileTimeAnalyzer/Controls/TimingDataGrid.xaml.cs b/Engine/Source/Programs/Windows/CompileTimeAnalyzer/CompileTimeAnalyzer/Controls/TimingDataGrid.xaml.cs
--- a/Engine/Source/Programs/Windows/CompileTimeAnalyzer/CompileTimeAnalyzer/Controls/TimingDataGrid.xaml.cs
+++ b/Engine/Source/Programs/Windows/CompileTimeAnalyzer/CompileTimeAnalyzer/Controls/TimingDataGrid.xaml.cs
@@ -15,6 +15,8 @@ namespace Timing_Data_Investigator.Controls
 	/// </summary>
 	public partial class TimingDataGrid : UserControl
 	{
+		private DataGridColumn SortingColumn = null;
+
 		public TimingDataGrid()
 		{
 			InitializeComponent();
@@ -43,33 +45,21 @@ namespace Timing_Data_Investigator.Controls
 					}
 			}
 
-			TreeGridFlatModel GridModel = (TreeGridFlatModel)DataContext;
-			PropertyInfo SortProperty = typeof(TimingDataViewModel).GetProperty(e.Column.SortMemberPath);
-			GridModel.Sort(SortProperty, e.Column.SortDirection);
-
+			SortingColumn = e.Column;
+			SortNodes();
 			e.Handled = true;
 		}
 
-		private void SortNodes(ObservableCollection<TreeGridElement> NewFlatModel, TreeGridElement Node, ListSortDirection SortDirection, PropertyInfo SortProperty)
+		private void SortNodes()
 		{
-			NewFlatModel.Add(Node);
-			if (Node.IsExpanded && Node.HasChildren)
+			if (SortingColumn == null)
 			{
-				IOrderedEnumerable<TreeGridElement> SortedChildren;
-				if (SortDirection == ListSortDirection.Ascending)
-				{
-					SortedChildren = Node.Children.OrderBy(c => SortProperty.GetValue(c));
-				}
-				else
-				{
-					SortedChildren = Node.Children.OrderByDescending(c => SortProperty.GetValue(c));
-				}
-
-				foreach (TreeGridElement Child in SortedChildren)
-				{
-					SortNodes(NewFlatModel, Child, SortDirection, SortProperty);
-				}
+				SortingColumn = Grid.Columns[2];
 			}
+
+			TreeGridFlatModel GridModel = (TreeGridFlatModel)DataContext;
+			PropertyInfo SortProperty = typeof(TimingDataViewModel).GetProperty(SortingColumn.SortMemberPath);
+			GridModel.Sort(SortProperty, SortingColumn.SortDirection);
 		}
 
 		private void Grid_KeyDown(object sender, KeyEventArgs e)
@@ -114,6 +104,10 @@ namespace Timing_Data_Investigator.Controls
 			}
 
 			SelectedData.IsExpanded = !SelectedData.IsExpanded;
+			if (SelectedData.IsExpanded)
+			{
+				SortNodes();
+			}
 		}
 	}
 }
diff --git a/Engine/Source/Programs/Windows/CompileTimeAnalyzer/CompileTimeAnalyzer/TreeGridFlatModel.cs b/Engine/Source/Programs/Windows/CompileTimeAnalyzer/CompileTimeAnalyzer/TreeGridFlatModel.cs
--- a/Engine/Source/Programs/Windows/CompileTimeAnalyzer/CompileTimeAnalyzer/TreeGridFlatModel.cs
+++ b/Engine/Source/Programs/Windows/CompileTimeAnalyzer/CompileTimeAnalyzer/TreeGridFlatModel.cs
@@ -114,17 +114,29 @@ namespace Timing_Data_Investigator
 			Modification = true;
 
 			// Iterate through all of the children within the items
-			IEnumerable<TreeGridElement> SortedItems = Items;
-			if (CurrentSortDirection == ListSortDirection.Ascending)
+			IEnumerable<TreeGridElement> SortedItems;
+			switch (CurrentSortDirection)
 			{
-				Items.OrderBy(i => CurrentSortProperty.GetValue(i));
-			}
-			else if (CurrentSortDirection == ListSortDirection.Descending)
-			{
-				Items.OrderByDescending(i => CurrentSortProperty.GetValue(i));
+				case ListSortDirection.Ascending:
+					{
+						SortedItems = Items.OrderBy(i => CurrentSortProperty.GetValue(i));
+						break;
+					}
+
+				case ListSortDirection.Descending:
+					{
+						SortedItems = Items.OrderByDescending(i => CurrentSortProperty.GetValue(i));
+						break;
+					}
+
+				default:
+					{
+						SortedItems = Items;
+						break;
+					}
 			}
 
-			foreach (TreeGridElement child in Items)
+			foreach (TreeGridElement child in SortedItems)
 			{
 				// Add the child to the model
 				Insert(Index++, child);
