Rename UE4PrereqSetup to UEPrereqSetup.

#rnx

diff --git a/Engine/Source/Programs/AutomationTool/Win/WinPlatform.Automation.cs b/Engine/Source/Programs/AutomationTool/Win/WinPlatform.Automation.cs
--- a/Engine/Source/Programs/AutomationTool/Win/WinPlatform.Automation.cs
+++ b/Engine/Source/Programs/AutomationTool/Win/WinPlatform.Automation.cs
@@ -129,7 +129,7 @@ public class Win64Platform : Platform
 
 		if (Params.Prereqs)
 		{
-			SC.StageFile(StagedFileType.NonUFS, FileReference.Combine(SC.EngineRoot, "Extras", "Redist", "en-us", "UE4PrereqSetup_x64.exe"));
+			SC.StageFile(StagedFileType.NonUFS, FileReference.Combine(SC.EngineRoot, "Extras", "Redist", "en-us", "UEPrereqSetup_x64.exe"));
 		}
 
 		if (!string.IsNullOrWhiteSpace(Params.AppLocalDirectory))
diff --git a/Engine/Source/Programs/Windows/BootstrapPackagedGame/Private/BootstrapPackagedGame.cpp b/Engine/Source/Programs/Windows/BootstrapPackagedGame/Private/BootstrapPackagedGame.cpp
--- a/Engine/Source/Programs/Windows/BootstrapPackagedGame/Private/BootstrapPackagedGame.cpp
+++ b/Engine/Source/Programs/Windows/BootstrapPackagedGame/Private/BootstrapPackagedGame.cpp
@@ -141,11 +141,11 @@ int InstallMissingPrerequisites(const WCHAR* BaseDirectory, const WCHAR* ExecDir
 		WCHAR PrereqInstaller[MAX_PATH];
 		if (bIsX64Target)
 		{
-			PathCombine(PrereqInstaller, BaseDirectory, L"Engine\\Extras\\Redist\\en-us\\UE4PrereqSetup_x64.exe");
+			PathCombine(PrereqInstaller, BaseDirectory, L"Engine\\Extras\\Redist\\en-us\\UEPrereqSetup_x64.exe");
 		}
 		else
 		{
-			PathCombine(PrereqInstaller, BaseDirectory, L"Engine\\Extras\\Redist\\en-us\\UE4PrereqSetup_x86.exe");
+			PathCombine(PrereqInstaller, BaseDirectory, L"Engine\\Extras\\Redist\\en-us\\UEPrereqSetup_x86.exe");
 		}
 		if(GetFileAttributes(PrereqInstaller) == INVALID_FILE_ATTRIBUTES)
 		{
