- Fixing copyrights

diff --git a/Engine/Source/Programs/Windows/UnrealWindowsForms/Gauntlet/GauntletBuildLauncherDialog.cs b/Engine/Source/Programs/Windows/UnrealWindowsForms/Gauntlet/GauntletBuildLauncherDialog.cs
--- a/Engine/Source/Programs/Windows/UnrealWindowsForms/Gauntlet/GauntletBuildLauncherDialog.cs
+++ b/Engine/Source/Programs/Windows/UnrealWindowsForms/Gauntlet/GauntletBuildLauncherDialog.cs
@@ -1,3 +1,5 @@
+// Copyright Epic Games, Inc. All Rights Reserved.
+
 using System;
 using System.Collections.Generic;
 using System.Text;
diff --git a/Engine/Source/Programs/Windows/UnrealWindowsForms/Turnkey/TurnkeyDialog.cs b/Engine/Source/Programs/Windows/UnrealWindowsForms/Turnkey/TurnkeyDialog.cs
--- a/Engine/Source/Programs/Windows/UnrealWindowsForms/Turnkey/TurnkeyDialog.cs
+++ b/Engine/Source/Programs/Windows/UnrealWindowsForms/Turnkey/TurnkeyDialog.cs
@@ -1,3 +1,5 @@
+// Copyright Epic Games, Inc. All Rights Reserved.
+
 using System;
 using System.Windows.Forms;
 
diff --git a/Engine/Source/Programs/Windows/UnrealWindowsForms/Utils/Utils.cs b/Engine/Source/Programs/Windows/UnrealWindowsForms/Utils/Utils.cs
--- a/Engine/Source/Programs/Windows/UnrealWindowsForms/Utils/Utils.cs
+++ b/Engine/Source/Programs/Windows/UnrealWindowsForms/Utils/Utils.cs
@@ -1,3 +1,5 @@
+// Copyright Epic Games, Inc. All Rights Reserved.
+
 using System;
 using System.Collections.Generic;
 using System.Text;
