- Fixed UnrealWindowForms.pdb to not be portable, which pdbcopy.exe can't handle

diff --git a/Engine/Source/Programs/Windows/UnrealWindowsForms/UnrealWindowsForms.csproj b/Engine/Source/Programs/Windows/UnrealWindowsForms/UnrealWindowsForms.csproj
--- a/Engine/Source/Programs/Windows/UnrealWindowsForms/UnrealWindowsForms.csproj
+++ b/Engine/Source/Programs/Windows/UnrealWindowsForms/UnrealWindowsForms.csproj
@@ -10,4 +10,9 @@
     <OutputPath>..\..\..\..\Binaries\DotNET\</OutputPath>
   </PropertyGroup>
 
+  <PropertyGroup Condition="'$(Configuration)|$(Platform)'=='Release|AnyCPU'">
+    <DebugType>pdbonly</DebugType>
+    <DebugSymbols>true</DebugSymbols>
+  </PropertyGroup>
+
 </Project>
