Fixing pdb format for UnrealWindowsForms

With the dotnet6 upgrade the pdb settings needed to be updated to produce a valid PDB for Microsofts pdbcopy

#rnx

diff --git a/Engine/Source/Programs/Windows/UnrealWindowsForms/UnrealWindowsForms.csproj b/Engine/Source/Programs/Windows/UnrealWindowsForms/UnrealWindowsForms.csproj
--- a/Engine/Source/Programs/Windows/UnrealWindowsForms/UnrealWindowsForms.csproj
+++ b/Engine/Source/Programs/Windows/UnrealWindowsForms/UnrealWindowsForms.csproj
@@ -8,6 +8,7 @@
     <UseWindowsForms>true</UseWindowsForms>
     <AppendTargetFrameworkToOutputPath>false</AppendTargetFrameworkToOutputPath>
     <OutputPath>..\..\..\..\Binaries\DotNET\</OutputPath>
+    <DebugType>none</DebugType>
   </PropertyGroup>
 
   <PropertyGroup Condition="'$(Configuration)|$(Platform)'=='Release|AnyCPU'">
