New Prereq installer for VS [REDACTED_PHONE] redist

#rnx

diff --git a/Engine/Source/Programs/Windows/BootstrapPackagedGame/Private/BootstrapPackagedGame.cpp b/Engine/Source/Programs/Windows/BootstrapPackagedGame/Private/BootstrapPackagedGame.cpp
--- a/Engine/Source/Programs/Windows/BootstrapPackagedGame/Private/BootstrapPackagedGame.cpp
+++ b/Engine/Source/Programs/Windows/BootstrapPackagedGame/Private/BootstrapPackagedGame.cpp
@@ -58,8 +58,8 @@ int InstallMissingPrerequisites(const WCHAR* BaseDirectory, const WCHAR* ExecDir
 	// Look for missing prerequisites
 	WCHAR MissingPrerequisites[1024] = { 0, };
 
-	// The Microsoft Visual C++ Runtime includes support for VS2015, VS2017, and VS2019
-	// https://docs.microsoft.com/en-us/cpp/windows/redistributing-visual-cpp-files?view=msvc-160
+	// The Microsoft Visual C++ Runtime includes support for VS2015, VS2017, VS2019, and VS2022
+	// https://docs.microsoft.com/en-us/cpp/windows/redistributing-visual-cpp-files?view=msvc-170
 
 	{
 		HKEY Hkey;
@@ -94,11 +94,11 @@ int InstallMissingPrerequisites(const WCHAR* BaseDirectory, const WCHAR* ExecDir
 			};
 
 			// This minimum should match the version installed by
-			// Engine/Source/Programs/PrereqInstaller/Resources/VCRedist/VC140/vcredist_x64.exe
+			// Engine/Source/Programs/PrereqInstaller/Resources/VCRedist/VC_redist.x64.exe
 			const DWORD RequiredMajor = 14;
-			const DWORD RequiredMinor = 24;
-			const DWORD RequiredBld = 28127;
-			const DWORD RequiredRbld = 4;
+			const DWORD RequiredMinor = 31;
+			const DWORD RequiredBld = 31103;
+			const DWORD RequiredRbld = 0;
 
 			const DWORD InstalledMajor = RegGetDwordOrZero(Hkey, L"Major");
 			const DWORD InstalledMinor = RegGetDwordOrZero(Hkey, L"Minor");
