[CompileTimeAnalyzer]
* Upgraded to .net 6 to make it work (since its dependencies are upgraded)
* Fixed a few compile errors
* Changed so CompileTimeAnalyzer looks for ".cta" files instead of ".timing.bin" files since that is what UBT is outputting these days
* Added bunch of pragma warning disable just to get through compiling

diff --git a/Engine/Source/Programs/Windows/CompileTimeAnalyzer/CompileTimeAnalyzer/App.config b/Engine/Source/Programs/Windows/CompileTimeAnalyzer/CompileTimeAnalyzer/App.config
deleted file mode 100644
diff --git a/Engine/Source/Programs/Windows/CompileTimeAnalyzer/CompileTimeAnalyzer/Commands/OpenTimingDataCommand.cs b/Engine/Source/Programs/Windows/CompileTimeAnalyzer/CompileTimeAnalyzer/Commands/OpenTimingDataCommand.cs
--- a/Engine/Source/Programs/Windows/CompileTimeAnalyzer/CompileTimeAnalyzer/Commands/OpenTimingDataCommand.cs
+++ b/Engine/Source/Programs/Windows/CompileTimeAnalyzer/CompileTimeAnalyzer/Commands/OpenTimingDataCommand.cs
@@ -1,5 +1,7 @@
 // Copyright Epic Games, Inc. All Rights Reserved.
 
+#pragma warning disable CS8618
+
 using System;
 using System.Windows.Input;
 using Timing_Data_Investigator.Models;
@@ -11,7 +13,7 @@ namespace Timing_Data_Investigator.Commands
 		private TimingDataViewModel ViewModelToOpen;
 
 #pragma warning disable CS0067
-		public event EventHandler CanExecuteChanged;
+		public event EventHandler? CanExecuteChanged;
 #pragma warning restore CS0067
 
 		public OpenTimingDataCommand(TimingDataViewModel TimingData)
@@ -21,12 +23,12 @@ namespace Timing_Data_Investigator.Commands
 
 		public Action<TimingDataViewModel> OpenAction { get; set; }
 
-		public bool CanExecute(object parameter)
+		public bool CanExecute(object? parameter)
 		{
 			return true;
 		}
 
-		public void Execute(object parameter)
+		public void Execute(object? parameter)
 		{
 			OpenAction?.Invoke(ViewModelToOpen);
 		}
diff --git a/Engine/Source/Programs/Windows/CompileTimeAnalyzer/CompileTimeAnalyzer/CompileTimeAnalyzer.csproj b/Engine/Source/Programs/Windows/CompileTimeAnalyzer/CompileTimeAnalyzer/CompileTimeAnalyzer.csproj
--- a/Engine/Source/Programs/Windows/CompileTimeAnalyzer/CompileTimeAnalyzer/CompileTimeAnalyzer.csproj
+++ b/Engine/Source/Programs/Windows/CompileTimeAnalyzer/CompileTimeAnalyzer/CompileTimeAnalyzer.csproj
@@ -1,136 +1,15 @@
-﻿<?xml version="1.0" encoding="utf-8"?>
-<Project ToolsVersion="15.0" xmlns="http://schemas.microsoft.com/developer/msbuild/2003">
-  <Import Project="$(MSBuildExtensionsPath)\$(MSBuildToolsVersion)\Microsoft.Common.props" Condition="Exists('$(MSBuildExtensionsPath)\$(MSBuildToolsVersion)\Microsoft.Common.props')" />
+﻿<Project Sdk="Microsoft.NET.Sdk">
+
   <PropertyGroup>
-    <Configuration Condition=" '$(Configuration)' == '' ">Debug</Configuration>
-    <Platform Condition=" '$(Platform)' == '' ">AnyCPU</Platform>
-    <ProjectGuid>{3F3BEAB2-B096-4420-9E90-BD0CB55867CE}</ProjectGuid>
     <OutputType>WinExe</OutputType>
-    <RootNamespace>CompileTimeAnalyzer</RootNamespace>
-    <AssemblyName>CompileTimeAnalyzer</AssemblyName>
-    <TargetFrameworkVersion>v4.6.2</TargetFrameworkVersion>
-    <FileAlignment>512</FileAlignment>
-    <ProjectTypeGuids>{60dc8134-eba5-43b8-bcc9-bb4bc16c2548};{FAE04EC0-301F-11D3-BF4B-00C04F79EFBC}</ProjectTypeGuids>
-    <WarningLevel>4</WarningLevel>
-    <AutoGenerateBindingRedirects>true</AutoGenerateBindingRedirects>
-    <Deterministic>true</Deterministic>
-    <TargetFrameworkProfile />
+    <TargetFramework>net6.0-windows</TargetFramework>
+    <Nullable>enable</Nullable>
+    <UseWPF>true</UseWPF>
   </PropertyGroup>
-  <PropertyGroup Condition=" '$(Configuration)|$(Platform)' == 'Debug|AnyCPU' ">
-    <PlatformTarget>x64</PlatformTarget>
-    <DebugSymbols>true</DebugSymbols>
-    <DebugType>full</DebugType>
-    <Optimize>false</Optimize>
-    <OutputPath>bin\Debug\</OutputPath>
-    <DefineConstants>DEBUG;TRACE</DefineConstants>
-    <ErrorReport>prompt</ErrorReport>
-    <WarningLevel>4</WarningLevel>
-  </PropertyGroup>
-  <PropertyGroup Condition=" '$(Configuration)|$(Platform)' == 'Release|AnyCPU' ">
-    <PlatformTarget>AnyCPU</PlatformTarget>
-    <DebugType>pdbonly</DebugType>
-    <Optimize>true</Optimize>
-    <OutputPath>bin\Release\</OutputPath>
-    <DefineConstants>TRACE</DefineConstants>
-    <ErrorReport>prompt</ErrorReport>
-    <WarningLevel>4</WarningLevel>
-  </PropertyGroup>
-  <PropertyGroup>
-    <ApplicationIcon>Default.ico</ApplicationIcon>
-  </PropertyGroup>
-  <ItemGroup>
-    <Reference Include="System" />
-    <Reference Include="System.Data" />
-    <Reference Include="System.Xml" />
-    <Reference Include="Microsoft.CSharp" />
-    <Reference Include="System.Core" />
-    <Reference Include="System.Xml.Linq" />
-    <Reference Include="System.Data.DataSetExtensions" />
-    <Reference Include="System.Net.Http" />
-    <Reference Include="System.Xaml">
-      <RequiredTargetFramework>4.0</RequiredTargetFramework>
-    </Reference>
-    <Reference Include="WindowsBase" />
-    <Reference Include="PresentationCore" />
-    <Reference Include="PresentationFramework" />
-  </ItemGroup>
-  <ItemGroup>
-    <ApplicationDefinition Include="App.xaml">
-      <Generator>MSBuild:Compile</Generator>
-      <SubType>Designer</SubType>
-    </ApplicationDefinition>
-    <Compile Include="Commands\OpenTimingDataCommand.cs" />
-    <Compile Include="Controls\TimingDataTab.xaml.cs">
-      <DependentUpon>TimingDataTab.xaml</DependentUpon>
-    </Compile>
-    <Compile Include="TreeGridElement.cs" />
-    <Compile Include="TreeGridElementComparer.cs" />
-    <Compile Include="TreeGridFlatModel.cs" />
-    <Compile Include="TreeGridModel.cs" />
-    <Page Include="Controls\TimingDataGrid.xaml">
-      <SubType>Designer</SubType>
-      <Generator>MSBuild:Compile</Generator>
-    </Page>
-    <Page Include="Controls\TimingDataTab.xaml">
-      <SubType>Designer</SubType>
-      <Generator>MSBuild:Compile</Generator>
-    </Page>
-    <Page Include="MainWindow.xaml">
-      <Generator>MSBuild:Compile</Generator>
-      <SubType>Designer</SubType>
-    </Page>
-    <Compile Include="App.xaml.cs">
-      <DependentUpon>App.xaml</DependentUpon>
-      <SubType>Code</SubType>
-    </Compile>
-    <Compile Include="Controls\TimingDataGrid.xaml.cs">
-      <DependentUpon>TimingDataGrid.xaml</DependentUpon>
-    </Compile>
-    <Compile Include="Converters.cs" />
-    <Compile Include="MainWindow.xaml.cs">
-      <DependentUpon>MainWindow.xaml</DependentUpon>
-      <SubType>Code</SubType>
-    </Compile>
-  </ItemGroup>
-  <ItemGroup>
-    <Compile Include="Models\TimingDataViewModel.cs" />
-    <Compile Include="Properties\AssemblyInfo.cs">
-      <SubType>Code</SubType>
-    </Compile>
-    <Compile Include="Properties\Resources.Designer.cs">
-      <AutoGen>True</AutoGen>
-      <DesignTime>True</DesignTime>
-      <DependentUpon>Resources.resx</DependentUpon>
-    </Compile>
-    <Compile Include="Properties\Settings.Designer.cs">
-      <AutoGen>True</AutoGen>
-      <DependentUpon>Settings.settings</DependentUpon>
-      <DesignTimeSharedInput>True</DesignTimeSharedInput>
-    </Compile>
-    <EmbeddedResource Include="Properties\Resources.resx">
-      <Generator>ResXFileCodeGenerator</Generator>
-      <LastGenOutput>Resources.Designer.cs</LastGenOutput>
-    </EmbeddedResource>
-    <None Include="Properties\Settings.settings">
-      <Generator>SettingsSingleFileGenerator</Generator>
-      <LastGenOutput>Settings.Designer.cs</LastGenOutput>
-    </None>
-  </ItemGroup>
-  <ItemGroup>
-    <None Include="App.config" />
-  </ItemGroup>
-  <ItemGroup>
-    <ProjectReference Include="..\..\..\Shared\EpicGames.Core\EpicGames.Core.csproj">
-      <Project>{5d7d66e8-8c76-4af9-b3ec-2ef03421d730}</Project>
-      <Name>DotNETUtilities</Name>
-    </ProjectReference>
-    <ProjectReference Include="..\..\..\UnrealBuildTool\UnrealBuildTool.csproj">
-      <Project>{fd7c5e1a-cfe4-4fd5-a525-1eb1599a39ac}</Project>
-      <Name>UnrealBuildTool</Name>
-    </ProjectReference>
-  </ItemGroup>
+
   <ItemGroup>
-    <Resource Include="Default.ico" />
+    <ProjectReference Include="..\..\..\Shared\EpicGames.Build\EpicGames.Build.csproj" />
+    <ProjectReference Include="..\..\..\UnrealBuildTool\UnrealBuildTool.csproj" />
   </ItemGroup>
-  <Import Project="$(MSBuildToolsPath)\Microsoft.CSharp.targets" />
-</Project>
\ No newline at end of file
+
+</Project>
diff --git a/Engine/Source/Programs/Windows/CompileTimeAnalyzer/CompileTimeAnalyzer/Controls/TimingDataGrid.xaml.cs b/Engine/Source/Programs/Windows/CompileTimeAnalyzer/CompileTimeAnalyzer/Controls/TimingDataGrid.xaml.cs
--- a/Engine/Source/Programs/Windows/CompileTimeAnalyzer/CompileTimeAnalyzer/Controls/TimingDataGrid.xaml.cs
+++ b/Engine/Source/Programs/Windows/CompileTimeAnalyzer/CompileTimeAnalyzer/Controls/TimingDataGrid.xaml.cs
@@ -1,5 +1,9 @@
 // Copyright Epic Games, Inc. All Rights Reserved.
 
+#pragma warning disable CS8625
+#pragma warning disable CS8600
+#pragma warning disable CS8604
+
 using System.Collections.ObjectModel;
 using System.ComponentModel;
 using System.Linq;
diff --git a/Engine/Source/Programs/Windows/CompileTimeAnalyzer/CompileTimeAnalyzer/Controls/TimingDataTab.xaml.cs b/Engine/Source/Programs/Windows/CompileTimeAnalyzer/CompileTimeAnalyzer/Controls/TimingDataTab.xaml.cs
--- a/Engine/Source/Programs/Windows/CompileTimeAnalyzer/CompileTimeAnalyzer/Controls/TimingDataTab.xaml.cs
+++ b/Engine/Source/Programs/Windows/CompileTimeAnalyzer/CompileTimeAnalyzer/Controls/TimingDataTab.xaml.cs
@@ -1,5 +1,12 @@
 // Copyright Epic Games, Inc. All Rights Reserved.
 
+#pragma warning disable CS8618
+#pragma warning disable CS8600
+#pragma warning disable CS8601
+#pragma warning disable CS8602
+#pragma warning disable CS8604
+#pragma warning disable CS8625
+
 using System.Collections.Generic;
 using System.ComponentModel;
 using System.Linq;
diff --git a/Engine/Source/Programs/Windows/CompileTimeAnalyzer/CompileTimeAnalyzer/Converters.cs b/Engine/Source/Programs/Windows/CompileTimeAnalyzer/CompileTimeAnalyzer/Converters.cs
--- a/Engine/Source/Programs/Windows/CompileTimeAnalyzer/CompileTimeAnalyzer/Converters.cs
+++ b/Engine/Source/Programs/Windows/CompileTimeAnalyzer/CompileTimeAnalyzer/Converters.cs
@@ -1,5 +1,7 @@
 // Copyright Epic Games, Inc. All Rights Reserved.
 
+#pragma warning disable CS8604
+
 using System;
 using System.Globalization;
 using System.Windows;
diff --git a/Engine/Source/Programs/Windows/CompileTimeAnalyzer/CompileTimeAnalyzer/MainWindow.xaml.cs b/Engine/Source/Programs/Windows/CompileTimeAnalyzer/CompileTimeAnalyzer/MainWindow.xaml.cs
--- a/Engine/Source/Programs/Windows/CompileTimeAnalyzer/CompileTimeAnalyzer/MainWindow.xaml.cs
+++ b/Engine/Source/Programs/Windows/CompileTimeAnalyzer/CompileTimeAnalyzer/MainWindow.xaml.cs
@@ -35,7 +35,7 @@ namespace Timing_Data_Investigator
 			string[] CmdArgs = Environment.GetCommandLineArgs();
 			if(CmdArgs.Length > 1)
 			{
-				if(CmdArgs[1].EndsWith(".timing.bin"))
+				if(CmdArgs[1].EndsWith(".cta"))
 				{
 					LoadTimingFile(CmdArgs[1]);
 					return;
@@ -115,7 +115,7 @@ namespace Timing_Data_Investigator
 		private void ShowOpenFileDialog()
 		{
 			OpenFileDialog OpenFileDialog = new OpenFileDialog();
-			OpenFileDialog.Filter = "Timing Files (*.timing.bin)|*.timing.bin|All Files (*.*)|*.*";
+			OpenFileDialog.Filter = "Timing Files (*.cta)|*.cta|All Files (*.*)|*.*";
 			if (OpenFileDialog.ShowDialog(this) == true)
 			{
 				LoadTimingFile(OpenFileDialog.FileName);
diff --git a/Engine/Source/Programs/Windows/CompileTimeAnalyzer/CompileTimeAnalyzer/Models/TimingDataViewModel.cs b/Engine/Source/Programs/Windows/CompileTimeAnalyzer/CompileTimeAnalyzer/Models/TimingDataViewModel.cs
--- a/Engine/Source/Programs/Windows/CompileTimeAnalyzer/CompileTimeAnalyzer/Models/TimingDataViewModel.cs
+++ b/Engine/Source/Programs/Windows/CompileTimeAnalyzer/CompileTimeAnalyzer/Models/TimingDataViewModel.cs
@@ -1,14 +1,20 @@
 // Copyright Epic Games, Inc. All Rights Reserved.
 
+#pragma warning disable CS8600
+#pragma warning disable CS8602
+#pragma warning disable CS8618
+#pragma warning disable CS8603
+#pragma warning disable CS8604
+
 using System.Collections.Generic;
 using System.Collections.ObjectModel;
 using System.ComponentModel;
 using System.IO;
 using System.IO.Compression;
 using System.Linq;
-using System.Windows.Input;
 using EpicGames.Core;
 using UnrealBuildTool;
+using ICommand = System.Windows.Input.ICommand;
 
 namespace Timing_Data_Investigator.Models
 {
@@ -143,7 +149,7 @@ namespace Timing_Data_Investigator.Models
 		public static TimingDataViewModel FromJsonFile(FileReference JsonFile)
 		{
 			string JsonString = File.ReadAllText(JsonFile.FullName);
-			TimingData DeserializedTimingData = Json.Deserialize<TimingData>(JsonString);
+			TimingData DeserializedTimingData = System.Text.Json.JsonSerializer.Deserialize<TimingData>(JsonString);
 			return FromTimingData(DeserializedTimingData);
 		}
 
diff --git a/Engine/Source/Programs/Windows/CompileTimeAnalyzer/CompileTimeAnalyzer/Properties/AssemblyInfo.cs b/Engine/Source/Programs/Windows/CompileTimeAnalyzer/CompileTimeAnalyzer/Properties/AssemblyInfo.cs
deleted file mode 100644
diff --git a/Engine/Source/Programs/Windows/CompileTimeAnalyzer/CompileTimeAnalyzer/TreeGridElement.cs b/Engine/Source/Programs/Windows/CompileTimeAnalyzer/CompileTimeAnalyzer/TreeGridElement.cs
--- a/Engine/Source/Programs/Windows/CompileTimeAnalyzer/CompileTimeAnalyzer/TreeGridElement.cs
+++ b/Engine/Source/Programs/Windows/CompileTimeAnalyzer/CompileTimeAnalyzer/TreeGridElement.cs
@@ -1,5 +1,12 @@
 // Copyright Epic Games, Inc. All Rights Reserved.
 
+#pragma warning disable CS8625
+#pragma warning disable CS8622
+#pragma warning disable CS8618
+#pragma warning disable CS8600
+#pragma warning disable CS8602
+#pragma warning disable CS8604
+
 using System;
 using System.Collections;
 using System.Collections.ObjectModel;
diff --git a/Engine/Source/Programs/Windows/CompileTimeAnalyzer/CompileTimeAnalyzer/TreeGridElementComparer.cs b/Engine/Source/Programs/Windows/CompileTimeAnalyzer/CompileTimeAnalyzer/TreeGridElementComparer.cs
--- a/Engine/Source/Programs/Windows/CompileTimeAnalyzer/CompileTimeAnalyzer/TreeGridElementComparer.cs
+++ b/Engine/Source/Programs/Windows/CompileTimeAnalyzer/CompileTimeAnalyzer/TreeGridElementComparer.cs
@@ -1,5 +1,8 @@
 // Copyright Epic Games, Inc. All Rights Reserved.
 
+#pragma warning disable CS8600
+#pragma warning disable CS8602
+
 using System;
 using System.Collections.Generic;
 using System.Linq;
@@ -18,7 +21,7 @@ namespace Timing_Data_Investigator
 			bSortDescending = SortDescending;
 		}
 
-		public int Compare(TClassType x, TClassType y)
+		public int Compare(TClassType? x, TClassType? y)
 		{
 			PropertyInfo ComparePropertyInfo = typeof(TClassType).GetProperties().FirstOrDefault(p => p.Name == sCompareProperty);
 			if (ComparePropertyInfo == null)
diff --git a/Engine/Source/Programs/Windows/CompileTimeAnalyzer/CompileTimeAnalyzer/TreeGridFlatModel.cs b/Engine/Source/Programs/Windows/CompileTimeAnalyzer/CompileTimeAnalyzer/TreeGridFlatModel.cs
--- a/Engine/Source/Programs/Windows/CompileTimeAnalyzer/CompileTimeAnalyzer/TreeGridFlatModel.cs
+++ b/Engine/Source/Programs/Windows/CompileTimeAnalyzer/CompileTimeAnalyzer/TreeGridFlatModel.cs
@@ -1,5 +1,8 @@
 // Copyright Epic Games, Inc. All Rights Reserved.
 
+#pragma warning disable CS8625
+#pragma warning disable CS8618
+
 using System;
 using System.Collections.Generic;
 using System.Collections.ObjectModel;
diff --git a/Engine/Source/Programs/Windows/CompileTimeAnalyzer/CompileTimeAnalyzer/TreeGridModel.cs b/Engine/Source/Programs/Windows/CompileTimeAnalyzer/CompileTimeAnalyzer/TreeGridModel.cs
--- a/Engine/Source/Programs/Windows/CompileTimeAnalyzer/CompileTimeAnalyzer/TreeGridModel.cs
+++ b/Engine/Source/Programs/Windows/CompileTimeAnalyzer/CompileTimeAnalyzer/TreeGridModel.cs
@@ -1,5 +1,9 @@
 // Copyright Epic Games, Inc. All Rights Reserved.
 
+#pragma warning disable CS8625
+#pragma warning disable CS8602
+#pragma warning disable CS8604
+
 using System.Collections;
 using System.Collections.Generic;
 using System.Collections.ObjectModel;
