BootstrapPackagedGame project is not build as non-unity.

 - no known reason why this project needed to be non-unity in the first place.
 - this means it won't get the 'NonUnity' suffix (from CL23421728) & Windows packaging code works again.

#rnx

diff --git a/Engine/Source/Programs/Windows/BootstrapPackagedGame/BootstrapPackagedGame.Target.cs b/Engine/Source/Programs/Windows/BootstrapPackagedGame/BootstrapPackagedGame.Target.cs
--- a/Engine/Source/Programs/Windows/BootstrapPackagedGame/BootstrapPackagedGame.Target.cs
+++ b/Engine/Source/Programs/Windows/BootstrapPackagedGame/BootstrapPackagedGame.Target.cs
@@ -15,7 +15,6 @@ public class BootstrapPackagedGameTarget : TargetRules
 
 		bUseStaticCRT = true;
 
-		bUseUnityBuild = false;
 		bUseSharedPCHs = false;
 		bUseMallocProfiler = false;
 
