Update prereq installer to install vcredist 14.34.31938

#rnx

diff --git a/Engine/Source/Programs/Windows/BootstrapPackagedGame/Private/BootstrapPackagedGame.cpp b/Engine/Source/Programs/Windows/BootstrapPackagedGame/Private/BootstrapPackagedGame.cpp
--- a/Engine/Source/Programs/Windows/BootstrapPackagedGame/Private/BootstrapPackagedGame.cpp
+++ b/Engine/Source/Programs/Windows/BootstrapPackagedGame/Private/BootstrapPackagedGame.cpp
@@ -97,7 +97,7 @@ int InstallMissingPrerequisites(const WCHAR* BaseDirectory, const WCHAR* ExecDir
 			// Engine/Source/Programs/PrereqInstaller/Resources/VCRedist/VC_redist.x64.exe
 			const DWORD RequiredMajor = 14;
 			const DWORD RequiredMinor = 34;
-			const DWORD RequiredBld = 31931;
+			const DWORD RequiredBld = 31938;
 			const DWORD RequiredRbld = 0;
 
 			const DWORD InstalledMajor = RegGetDwordOrZero(Hkey, L"Major");
