BootstrapPackagedGame: Allow skipping the prereq installer if the required dependencies are bundled with the application

* Query the file version of dlls to be loaded so it can be compared against the min required version, in addition to ensuring the dll can be loaded
* Unload dlls after verifying they can be loaded
* Check bundled dlls first as dlls from an application's directory will be preferred when LoadLibrary is called.
* Check version of directx dlls
* Remove x86 logic, Win32 hasn't been supported for years
* Fix static analysis warnings

diff --git a/Engine/Source/Programs/Windows/BootstrapPackagedGame/BootstrapPackagedGame.Build.cs b/Engine/Source/Programs/Windows/BootstrapPackagedGame/BootstrapPackagedGame.Build.cs
--- a/Engine/Source/Programs/Windows/BootstrapPackagedGame/BootstrapPackagedGame.Build.cs
+++ b/Engine/Source/Programs/Windows/BootstrapPackagedGame/BootstrapPackagedGame.Build.cs
@@ -7,6 +7,6 @@ public class BootstrapPackagedGame : ModuleRules
 	public BootstrapPackagedGame(ReadOnlyTargetRules Target) : base(Target)
 	{
 		PublicSystemLibraries.Add("shlwapi.lib");
-		bEnableUndefinedIdentifierWarnings = false;
+		PublicSystemLibraries.Add("version.lib");
 	}
 }
diff --git a/Engine/Source/Programs/Windows/BootstrapPackagedGame/BootstrapPackagedGame.Target.cs b/Engine/Source/Programs/Windows/BootstrapPackagedGame/BootstrapPackagedGame.Target.cs
--- a/Engine/Source/Programs/Windows/BootstrapPackagedGame/BootstrapPackagedGame.Target.cs
+++ b/Engine/Source/Programs/Windows/BootstrapPackagedGame/BootstrapPackagedGame.Target.cs
@@ -11,6 +11,7 @@ public class BootstrapPackagedGameTarget : TargetRules
 	{
 		Type = TargetType.Program;
 		LinkType = TargetLinkType.Monolithic;
+		IncludeOrderVersion = EngineIncludeOrderVersion.Latest;
 		LaunchModuleName = "BootstrapPackagedGame";
 
 		bUseStaticCRT = true;
diff --git a/Engine/Source/Programs/Windows/BootstrapPackagedGame/Private/BootstrapPackagedGame.cpp b/Engine/Source/Programs/Windows/BootstrapPackagedGame/Private/BootstrapPackagedGame.cpp
--- a/Engine/Source/Programs/Windows/BootstrapPackagedGame/Private/BootstrapPackagedGame.cpp
+++ b/Engine/Source/Programs/Windows/BootstrapPackagedGame/Private/BootstrapPackagedGame.cpp
@@ -5,6 +5,28 @@
 #define IDI_EXEC_FILE 201
 #define IDI_EXEC_ARGS 202
 
+struct VersionInfo
+{
+	DWORD Major = 0;
+	DWORD Minor = 0;
+	DWORD Bld = 0;
+	DWORD Rbld = 0;
+};
+
+// This minimum should match the version installed by
+// Engine/Source/Programs/PrereqInstaller/Resources/VCRedist/VC_redist.x64.exe
+static const VersionInfo MinRedistVersion = { 14, 36, 32532, 0 };
+static const VersionInfo MinXInputVersion = { 9, 18, 944, 0 };
+
+bool IsVersionValid(const VersionInfo& Version, const VersionInfo& MinVersion)
+{
+	if (Version.Major > MinVersion.Major) return true;
+	if (Version.Major == MinVersion.Major && Version.Minor > MinVersion.Minor) return true;
+	if (Version.Major == MinVersion.Major && Version.Minor == MinVersion.Minor && Version.Bld > MinVersion.Bld) return true;
+	if (Version.Major == MinVersion.Major && Version.Minor == MinVersion.Minor && Version.Bld == MinVersion.Bld && Version.Rbld >= MinVersion.Rbld) return true;
+	return false;
+}
+
 WCHAR* ReadResourceString(HMODULE ModuleHandle, LPCWSTR Name)
 {
 	WCHAR* Result = NULL;
@@ -29,106 +51,160 @@ WCHAR* ReadResourceString(HMODULE ModuleHandle, LPCWSTR Name)
 
 bool TryLoadDll(const WCHAR* ExecDirectory, const WCHAR* Name)
 {
-	// Try to load it from the system path
-	if (LoadLibrary(Name) != nullptr)
+	WCHAR AppLocalPath[MAX_PATH];
+	if (PathCombine(AppLocalPath, ExecDirectory, Name) == nullptr)
+	{
+		return false;
+	}
+	HMODULE Handle = LoadLibrary(AppLocalPath);
+	if (Handle != nullptr)
 	{
+		FreeLibrary(Handle);
 		return true;
 	}
+	return false;
+}
 
-	// Try to load it from the application directory
-	WCHAR AppLocalPath[MAX_PATH];
-	PathCombine(AppLocalPath, ExecDirectory, Name);
-	if (LoadLibrary(AppLocalPath) != nullptr)
+bool TryLoadDll(const WCHAR* Name)
+{
+	return TryLoadDll(nullptr, Name);
+}
+
+bool TryGetFileVersionInfo(const WCHAR* ExecDirectory, const WCHAR* Name, VersionInfo& outVersionInfo)
+{
+	WCHAR Path[MAX_PATH];
+	if (PathCombine(Path, ExecDirectory, Name) == nullptr)
 	{
-		return true;
+		return false;
 	}
 
-	// Otherwise fail
-	return false;
+	DWORD VersionSize = GetFileVersionInfoSize(Path, nullptr);
+	if (VersionSize == 0)
+	{
+		return false;
+	}
+
+	LPSTR pVersionData = new char[VersionSize];
+	if (!GetFileVersionInfo(Path, 0, VersionSize, pVersionData))
+	{
+		delete[] pVersionData;
+		return false;
+	}
+
+	VS_FIXEDFILEINFO* pFileInfo = nullptr;
+	UINT FileInfoLen = 0;
+	if (!VerQueryValue(pVersionData, L"\\", (LPVOID*)&pFileInfo, &FileInfoLen) && FileInfoLen != 0)
+	{
+		delete[] pVersionData;
+		return false;
+	}
+
+	outVersionInfo.Major = (pFileInfo->dwFileVersionMS >> 16) & 0xffff;
+	outVersionInfo.Minor = (pFileInfo->dwFileVersionMS >> 0) & 0xffff;
+	outVersionInfo.Bld = (pFileInfo->dwFileVersionLS >> 16) & 0xffff;
+	outVersionInfo.Rbld = (pFileInfo->dwFileVersionLS >> 0) & 0xffff;
+
+	delete[] pVersionData;
+	return true;
 }
 
-int InstallMissingPrerequisites(const WCHAR* BaseDirectory, const WCHAR* ExecDirectory)
+bool TryGetFileVersionInfo(const WCHAR* Name, VersionInfo& outVersionInfo)
 {
-#ifdef _M_X64
-	bool bIsX64Target = true;
-#else
-	bool bIsX64Target = false;
-#endif
+	return TryGetFileVersionInfo(nullptr, Name, outVersionInfo);
+}
+
+bool IsDllValid(const WCHAR* ExecDirectory, const WCHAR* Name, const VersionInfo& RequiredVersion)
+{
+	VersionInfo DllInfo;
+	return TryGetFileVersionInfo(ExecDirectory, Name, DllInfo) && IsVersionValid(DllInfo, RequiredVersion) && TryLoadDll(ExecDirectory, Name);
+}
 
+bool IsDllValid(const WCHAR* Name, const VersionInfo& RequiredVersion)
+{
+	return IsDllValid(nullptr, Name, RequiredVersion);
+}
+
+int InstallMissingPrerequisites(const WCHAR* BaseDirectory, const WCHAR* ExecDirectory)
+{
 	// Look for missing prerequisites
 	WCHAR MissingPrerequisites[1024] = { 0, };
 
 	// The Microsoft Visual C++ Runtime includes support for VS2015, VS2017, VS2019, and VS2022
 	// https://docs.microsoft.com/en-us/cpp/windows/redistributing-visual-cpp-files?view=msvc-170
-
+	
 	{
-		HKEY Hkey;
-		LSTATUS KeyOpenStatus;
+		bool bInstallVCRedist = true;
 
-		if (bIsX64Target)
+		// Check the file version of bundled redist dlls
+		if (IsDllValid(ExecDirectory, L"msvcp140_2.dll", MinRedistVersion) &&
+			IsDllValid(ExecDirectory, L"vcruntime140_1.dll", MinRedistVersion))
 		{
-			KeyOpenStatus = RegOpenKeyExW(HKEY_LOCAL_MACHINE, L"SOFTWARE\\Microsoft\\VisualStudio\\14.0\\VC\\Runtimes\\x64", 0, KEY_READ, &Hkey);
+			bInstallVCRedist = false;
 		}
-		else
+
+		// If no bundled redist dlls are available, check the registry for the installed redist, 
+		if (bInstallVCRedist)
 		{
-			// 32bit build running on an 32bit host
-			KeyOpenStatus = RegOpenKeyExW(HKEY_LOCAL_MACHINE, L"SOFTWARE\\Microsoft\\VisualStudio\\14.0\\VC\\Runtimes\\x86", 0, KEY_READ, &Hkey);
+			HKEY Hkey;
+			LSTATUS KeyOpenStatus = RegOpenKeyExW(HKEY_LOCAL_MACHINE, L"SOFTWARE\\Microsoft\\VisualStudio\\14.0\\VC\\Runtimes\\x64", 0, KEY_READ, &Hkey);
 
-			if (ERROR_SUCCESS != KeyOpenStatus)
+			if (KeyOpenStatus == ERROR_SUCCESS)
+			{
+				auto RegGetDwordOrZero = [](HKEY Hkey, LPCWSTR Name) -> DWORD
+					{
+						DWORD Value = 0;
+						DWORD ValueSize = sizeof Value;
+						LSTATUS Status = RegQueryValueExW(Hkey, Name, NULL, NULL, (LPBYTE)&Value, &ValueSize);
+						return ERROR_SUCCESS == Status ? Value : 0;
+					};
+
+				const VersionInfo InstalledVersion = {
+					RegGetDwordOrZero(Hkey, L"Major"),
+					RegGetDwordOrZero(Hkey, L"Minor"),
+					RegGetDwordOrZero(Hkey, L"Bld"),
+					RegGetDwordOrZero(Hkey, L"Rbld")
+				};
+
+				RegCloseKey(Hkey);
+
+				if (IsVersionValid(InstalledVersion, MinRedistVersion))
+				{
+					// it is possible that the redist has been uninstalled but the registry entries have not been removed
+					// test that some relatively new dlls are able to be loaded
+					if (IsDllValid(L"msvcp140_2.dll", MinRedistVersion) &&
+						IsDllValid(L"vcruntime140_1.dll", MinRedistVersion))
+					{
+						bInstallVCRedist = false;
+					}
+				}
+			}
+			if (bInstallVCRedist)
 			{
-				// 32bit build running on 64bit host
-				KeyOpenStatus = RegOpenKeyExW(HKEY_LOCAL_MACHINE, L"SOFTWARE\\Wow6432Node\\Microsoft\\VisualStudio\\14.0\\VC\\Runtimes\\x86", 0, KEY_READ, &Hkey);
+				wcscat_s(MissingPrerequisites, TEXT("Microsoft Visual C++ Runtime\n"));
 			}
 		}
+	}
 
-		bool bInstallVCRedist = true;
+	{
+		bool bInstallDirectXRuntime = true;
 
-		if (KeyOpenStatus == ERROR_SUCCESS)
+		// Check the file version of bundled directx dlls
+		if (IsDllValid(ExecDirectory, L"xinput1_3.dll", MinXInputVersion))
 		{
-			auto RegGetDwordOrZero = [](HKEY Hkey, LPCWSTR Name) -> DWORD
-			{
-				DWORD Value = 0;
-				DWORD ValueSize = sizeof Value;
-				LSTATUS Status = RegQueryValueExW(Hkey, Name, NULL, NULL, (LPBYTE)&Value, &ValueSize);
-				return ERROR_SUCCESS == Status ? Value : 0;
-			};
-
-			// This minimum should match the version installed by
-			// Engine/Source/Programs/PrereqInstaller/Resources/VCRedist/VC_redist.x64.exe
-			const DWORD RequiredMajor = 14;
-			const DWORD RequiredMinor = 36;
-			const DWORD RequiredBld = 32532;
-			const DWORD RequiredRbld = 0;
-
-			const DWORD InstalledMajor = RegGetDwordOrZero(Hkey, L"Major");
-			const DWORD InstalledMinor = RegGetDwordOrZero(Hkey, L"Minor");
-			const DWORD InstalledBld = RegGetDwordOrZero(Hkey, L"Bld");
-			const DWORD InstalledRbld = RegGetDwordOrZero(Hkey, L"Rbld");
-
-			if ((InstalledMajor > RequiredMajor) ||
-				(InstalledMajor == RequiredMajor && InstalledMinor > RequiredMinor) ||
-				(InstalledMajor == RequiredMajor && InstalledMinor == RequiredMinor && InstalledBld > RequiredBld) ||
-				(InstalledMajor == RequiredMajor && InstalledMinor == RequiredMinor && InstalledBld == RequiredBld && InstalledRbld >= RequiredRbld))
-			{
-				// it is possible that the redist has been uninstalled but the registry entries have not been removed
-				// test that some relatively new dlls are able to be loaded
-				if (TryLoadDll(ExecDirectory, L"msvcp140_2.dll") &&
-					(!bIsX64Target || TryLoadDll(ExecDirectory, L"vcruntime140_1.dll")))
-				{
-					bInstallVCRedist = false;
-				}
-			}
-			RegCloseKey(Hkey);
+			bInstallDirectXRuntime = false;
 		}
-		if (bInstallVCRedist)
+
+		// Check the file version of system dirextx dlls
+		if (bInstallDirectXRuntime && IsDllValid(L"xinput1_3.dll", MinXInputVersion))
 		{
-			wcscat_s(MissingPrerequisites, TEXT("Microsoft Visual C++ Runtime\n"));
+			bInstallDirectXRuntime = false;
 		}
-	}
 
-	if(!TryLoadDll(ExecDirectory, L"XINPUT1_3.DLL"))
-	{
-		wcscat_s(MissingPrerequisites, TEXT("DirectX Runtime\n"));
+		// No valid directx dll was found
+		if (bInstallDirectXRuntime)
+		{
+			wcscat_s(MissingPrerequisites, TEXT("DirectX Runtime\n"));
+		}
 	}
 
 	// Check if there's anything missing
@@ -139,14 +215,7 @@ int InstallMissingPrerequisites(const WCHAR* BaseDirectory, const WCHAR* ExecDir
 
 		// If we don't have the installer, just notify the user and quit
 		WCHAR PrereqInstaller[MAX_PATH];
-		if (bIsX64Target)
-		{
-			PathCombine(PrereqInstaller, BaseDirectory, L"Engine\\Extras\\Redist\\en-us\\UEPrereqSetup_x64.exe");
-		}
-		else
-		{
-			PathCombine(PrereqInstaller, BaseDirectory, L"Engine\\Extras\\Redist\\en-us\\UEPrereqSetup_x86.exe");
-		}
+		PathCombine(PrereqInstaller, BaseDirectory, L"Engine\\Extras\\Redist\\en-us\\UEPrereqSetup_x64.exe");
 		if(GetFileAttributes(PrereqInstaller) == INVALID_FILE_ATTRIBUTES)
 		{
 			MessageBox(NULL, MissingPrerequisitesMsg, NULL, MB_OK);
@@ -215,14 +284,14 @@ int SpawnTarget(WCHAR* CmdLine)
 	return (int)ExitCode;
 }
 
-int WINAPI wWinMain(HINSTANCE hInstance, HINSTANCE hPrevInstance, TCHAR* CmdLine, int ShowCmd)
+int WINAPI wWinMain(_In_ HINSTANCE hInstance, _In_opt_ HINSTANCE hPrevInstance, _In_ LPWSTR CmdLine, _In_ int ShowCmd)
 {
 	(void)hPrevInstance;
 	(void)ShowCmd;
 
 	// Get the current module filename
 	WCHAR CurrentModuleFile[MAX_PATH];
-	GetModuleFileNameW(hInstance, CurrentModuleFile, sizeof(CurrentModuleFile));
+	GetModuleFileNameW(hInstance, CurrentModuleFile, MAX_PATH);
 
 	// Get the base directory from the current module filename
 	WCHAR BaseDirectory[MAX_PATH];
diff --git a/Engine/Source/Programs/Windows/BootstrapPackagedGame/Private/BootstrapPackagedGame.h b/Engine/Source/Programs/Windows/BootstrapPackagedGame/Private/BootstrapPackagedGame.h
--- a/Engine/Source/Programs/Windows/BootstrapPackagedGame/Private/BootstrapPackagedGame.h
+++ b/Engine/Source/Programs/Windows/BootstrapPackagedGame/Private/BootstrapPackagedGame.h
@@ -12,3 +12,4 @@ struct IUnknown;
 #include <assert.h>
 #include <stdio.h>
 #include <shlwapi.h>
+#include <winver.h>
