try to fix the Static Analysis warning from windows in BootstrapPackagedGame

#rnx

diff --git a/Engine/Source/Programs/Windows/BootstrapPackagedGame/Private/BootstrapPackagedGame.h b/Engine/Source/Programs/Windows/BootstrapPackagedGame/Private/BootstrapPackagedGame.h
--- a/Engine/Source/Programs/Windows/BootstrapPackagedGame/Private/BootstrapPackagedGame.h
+++ b/Engine/Source/Programs/Windows/BootstrapPackagedGame/Private/BootstrapPackagedGame.h
@@ -4,10 +4,21 @@
 
 struct IUnknown;
 
+// Disables a known issue in "winreg.h" when running static analysis with /analyze
+// The warning is as follows:
+//
+//		warning C6553: The annotation for function 'RegOpenKeyExW' on _Param_(3) does not apply to a value type.
+//
+// see https://developercommunity.visualstudio.com/t/warning-C6553:-The-annotation-for-functi/1676659 for more details
+#pragma warning(push)
+#pragma warning(disable : 6553)
+
 #pragma pack(push, 8)
 #include <windows.h>
 #pragma pack(pop)
 
+#pragma warning(pop) 
+
 #include <tchar.h>
 #include <assert.h>
 #include <stdio.h>
