Update Bootstrap package redist version to 14.38.33130

#rnx

diff --git a/Engine/Source/Programs/Windows/BootstrapPackagedGame/Private/BootstrapPackagedGame.cpp b/Engine/Source/Programs/Windows/BootstrapPackagedGame/Private/BootstrapPackagedGame.cpp
--- a/Engine/Source/Programs/Windows/BootstrapPackagedGame/Private/BootstrapPackagedGame.cpp
+++ b/Engine/Source/Programs/Windows/BootstrapPackagedGame/Private/BootstrapPackagedGame.cpp
@@ -15,7 +15,7 @@ struct VersionInfo
 
 // This minimum should match the version installed by
 // Engine/Source/Programs/PrereqInstaller/Resources/VCRedist/VC_redist.x64.exe
-static const VersionInfo MinRedistVersion = { 14, 36, 32532, 0 };
+static const VersionInfo MinRedistVersion = { 14, 38, 33130, 0 };
 
 bool IsVersionValid(const VersionInfo& Version, const VersionInfo& MinVersion)
 {
