BootstrapPackagedGame: When checking for the installed redist, also directly check system32 dlls in addition to by filename only

#rnx

diff --git a/Engine/Source/Programs/Windows/BootstrapPackagedGame/Private/BootstrapPackagedGame.cpp b/Engine/Source/Programs/Windows/BootstrapPackagedGame/Private/BootstrapPackagedGame.cpp
--- a/Engine/Source/Programs/Windows/BootstrapPackagedGame/Private/BootstrapPackagedGame.cpp
+++ b/Engine/Source/Programs/Windows/BootstrapPackagedGame/Private/BootstrapPackagedGame.cpp
@@ -169,8 +169,18 @@ int InstallMissingPrerequisites(const WCHAR* BaseDirectory, const WCHAR* ExecDir
 				if (IsVersionValid(InstalledVersion, MinRedistVersion))
 				{
 					// it is possible that the redist has been uninstalled but the registry entries have not been removed
-					// test that some relatively new dlls are able to be loaded
-					if (IsDllValid(L"msvcp140_2.dll", MinRedistVersion) &&
+					// test that some relatively new dlls are able to be loaded from system32
+					WCHAR SystemRoot[MAX_PATH] = { 0 };
+					GetEnvironmentVariable(L"SystemRoot", SystemRoot, MAX_PATH);
+					WCHAR System32Path[MAX_PATH] = { 0 };
+					PathCombine(System32Path, SystemRoot, L"system32");
+					if (IsDllValid(System32Path, L"msvcp140_2.dll", MinRedistVersion) &&
+						IsDllValid(System32Path, L"vcruntime140_1.dll", MinRedistVersion))
+					{
+						bInstallVCRedist = false;
+					}
+					// test that some relatively new dlls are able to be loaded with no path if not found in system32
+					else if (IsDllValid(L"msvcp140_2.dll", MinRedistVersion) &&
 						IsDllValid(L"vcruntime140_1.dll", MinRedistVersion))
 					{
 						bInstallVCRedist = false;
