Bootstrap packaged game will skip VC runtime DLL install if there is an appx package dependency on the VCLibs.

#rnx

diff --git a/Engine/Source/Programs/Windows/BootstrapPackagedGame/Private/BootstrapPackagedGame.cpp b/Engine/Source/Programs/Windows/BootstrapPackagedGame/Private/BootstrapPackagedGame.cpp
--- a/Engine/Source/Programs/Windows/BootstrapPackagedGame/Private/BootstrapPackagedGame.cpp
+++ b/Engine/Source/Programs/Windows/BootstrapPackagedGame/Private/BootstrapPackagedGame.cpp
@@ -123,6 +123,49 @@ bool IsDllValid(const WCHAR* Name, const VersionInfo& RequiredVersion)
 	return IsDllValid(nullptr, Name, RequiredVersion);
 }
 
+bool HasAppxPackagedVCRuntime()
+{
+	static const WCHAR* PackageFamilyNameVCLibs = TEXT("Microsoft.VCLibs.140.00.UWPDesktop_8wekyb3d8bbwe");
+
+	// try to find the GetCurrentPackageInfo function
+	HMODULE hModule = GetModuleHandleW(TEXT("kernel32.dll"));
+	typedef LONG(WINAPI *GetCurrentPackageInfoProc)(const UINT32, UINT32*, BYTE*, UINT32*);
+	GetCurrentPackageInfoProc fnGetCurrentPackageInfo = hModule ? (GetCurrentPackageInfoProc)GetProcAddress(hModule, "GetCurrentPackageInfo") : nullptr;
+
+	// attempt to enumerate the package dependencies & check if there is a dependency on the VCLibs package
+	bool bHasVCLibs = false;
+	if (fnGetCurrentPackageInfo != nullptr)
+	{
+		UINT32 BufferLength = 0;
+		LONG Result = fnGetCurrentPackageInfo(PACKAGE_FILTER_DIRECT, &BufferLength, nullptr, nullptr);
+
+		if (Result == ERROR_INSUFFICIENT_BUFFER)
+		{
+			UINT32 Count = 0;
+			PACKAGE_INFO* PackageInfo = (PACKAGE_INFO*)malloc(BufferLength);
+			Result = fnGetCurrentPackageInfo(PACKAGE_FILTER_DIRECT, &BufferLength, (BYTE*)PackageInfo, &Count);
+
+			if (Result == ERROR_SUCCESS && PackageInfo != nullptr)
+			{
+				for (UINT32 Index = 0; Index < Count; Index++)
+				{
+					if (wcscmp(PackageInfo[Index].packageFamilyName, PackageFamilyNameVCLibs ) == 0)
+					{
+						// note: not checking PackageInfo[Index].packageId.version against MinRedistVersion because unfortunately the Windows Store version trails behind MSVC
+						bHasVCLibs = true;
+						break;
+					}
+				}
+			}
+			free(PackageInfo);
+		}
+	}
+
+	return bHasVCLibs;
+}
+
+
+
 int InstallMissingPrerequisites(const WCHAR* BaseDirectory, const WCHAR* ExecDirectory)
 {
 	// Look for missing prerequisites
@@ -141,6 +184,12 @@ int InstallMissingPrerequisites(const WCHAR* BaseDirectory, const WCHAR* ExecDir
 			bInstallVCRedist = false;
 		}
 
+		// Check if we are part of an appx package with an embedded dependency on the VC runtime libraries
+		if (bInstallVCRedist && HasAppxPackagedVCRuntime())
+		{
+			bInstallVCRedist = false;
+		}
+
 		// If no bundled redist dlls are available, check the registry for the installed redist, 
 		if (bInstallVCRedist)
 		{
diff --git a/Engine/Source/Programs/Windows/BootstrapPackagedGame/Private/BootstrapPackagedGame.h b/Engine/Source/Programs/Windows/BootstrapPackagedGame/Private/BootstrapPackagedGame.h
--- a/Engine/Source/Programs/Windows/BootstrapPackagedGame/Private/BootstrapPackagedGame.h
+++ b/Engine/Source/Programs/Windows/BootstrapPackagedGame/Private/BootstrapPackagedGame.h
@@ -24,3 +24,4 @@ struct IUnknown;
 #include <stdio.h>
 #include <shlwapi.h>
 #include <winver.h>
+#include <appmodel.h>
