Update BoostrapPackagedGame to run vc_redist directly

#rnx

diff --git a/Engine/Source/Programs/Windows/BootstrapPackagedGame/Private/BootstrapPackagedGame.cpp b/Engine/Source/Programs/Windows/BootstrapPackagedGame/Private/BootstrapPackagedGame.cpp
--- a/Engine/Source/Programs/Windows/BootstrapPackagedGame/Private/BootstrapPackagedGame.cpp
+++ b/Engine/Source/Programs/Windows/BootstrapPackagedGame/Private/BootstrapPackagedGame.cpp
@@ -14,8 +14,8 @@ struct VersionInfo
 };
 
 // This minimum should match the version installed by
-// Engine/Source/Programs/PrereqInstaller/Resources/VCRedist/VC_redist.x64.exe
-static const VersionInfo MinRedistVersion = { 14, 40, 33810, 0 };
+// Engine/Extras/Redist/en-us/vc_redist.x64.exe
+static const VersionInfo MinRedistVersion = { 14, 42, 34438, 0 };
 
 bool IsVersionValid(const VersionInfo& Version, const VersionInfo& MinVersion)
 {
@@ -178,7 +178,8 @@ int InstallMissingPrerequisites(const WCHAR* BaseDirectory, const WCHAR* ExecDir
 		bool bInstallVCRedist = true;
 
 		// Check the file version of bundled redist dlls
-		if (IsDllValid(ExecDirectory, L"msvcp140_2.dll", MinRedistVersion) &&
+		if (ExecDirectory != nullptr &&
+			IsDllValid(ExecDirectory, L"msvcp140_2.dll", MinRedistVersion) &&
 			IsDllValid(ExecDirectory, L"vcruntime140_1.dll", MinRedistVersion))
 		{
 			bInstallVCRedist = false;
@@ -238,7 +239,7 @@ int InstallMissingPrerequisites(const WCHAR* BaseDirectory, const WCHAR* ExecDir
 			}
 			if (bInstallVCRedist)
 			{
-				wcscat_s(MissingPrerequisites, TEXT("Microsoft Visual C++ Runtime\n"));
+				wcscat_s(MissingPrerequisites, TEXT("Microsoft Visual C++ 2015-2022 Redistributable (x64)\n"));
 			}
 		}
 	}
@@ -252,7 +253,7 @@ int InstallMissingPrerequisites(const WCHAR* BaseDirectory, const WCHAR* ExecDir
 
 		// If we don't have the installer, just notify the user and quit
 		WCHAR PrereqInstaller[MAX_PATH];
-		PathCombine(PrereqInstaller, BaseDirectory, L"Engine\\Extras\\Redist\\en-us\\UEPrereqSetup_x64.exe");
+		PathCombine(PrereqInstaller, BaseDirectory, L"Engine\\Extras\\Redist\\en-us\\vc_redist.x64.exe");
 		if(GetFileAttributes(PrereqInstaller) == INVALID_FILE_ATTRIBUTES)
 		{
 			MessageBox(NULL, MissingPrerequisitesMsg, NULL, MB_OK);
@@ -266,6 +267,9 @@ int InstallMissingPrerequisites(const WCHAR* BaseDirectory, const WCHAR* ExecDir
 			return 9002;
 		}
 
+		WCHAR PrereqParameters[1024] = { 0, };
+		wcscat_s(PrereqParameters, L"/passive /norestart");
+
 		// Start the installer
 		SHELLEXECUTEINFO ShellExecuteInfo;
 		ZeroMemory(&ShellExecuteInfo, sizeof(ShellExecuteInfo));
@@ -273,6 +277,7 @@ int InstallMissingPrerequisites(const WCHAR* BaseDirectory, const WCHAR* ExecDir
 		ShellExecuteInfo.fMask = SEE_MASK_NOCLOSEPROCESS;
 		ShellExecuteInfo.nShow = SW_SHOWNORMAL;
 		ShellExecuteInfo.lpFile = PrereqInstaller;
+		ShellExecuteInfo.lpParameters = PrereqParameters;
 		if(!ShellExecuteExW(&ShellExecuteInfo))
 		{
 			return 9003;
@@ -283,7 +288,8 @@ int InstallMissingPrerequisites(const WCHAR* BaseDirectory, const WCHAR* ExecDir
 		WaitForSingleObject(ShellExecuteInfo.hProcess, INFINITE);
 		GetExitCodeProcess(ShellExecuteInfo.hProcess, &ExitCode);
 		CloseHandle(ShellExecuteInfo.hProcess);
-		if(ExitCode != 0)
+		// 1638: Newer version already installed
+		if(ExitCode != 0 && ExitCode != 1638)
 		{
 			return 9004;
 		}
diff --git a/Setup.bat b/Setup.bat
--- a/Setup.bat
+++ b/Setup.bat
@@ -24,7 +24,7 @@ echo Engine/Binaries/DotNET/GitDependencies/win-x64/GitDependencies.exe %* >>.gi
 
 rem Install prerequisites...
 echo Installing prerequisites...
-start /wait Engine\Extras\Redist\en-us\UEPrereqSetup_x64.exe /quiet /norestart
+start /wait Engine\Extras\Redist\en-us\vc_redist.x64.exe /quiet /norestart
 
 rem Register the engine installation...
 if not exist .\Engine\Binaries\Win64\UnrealVersionSelector-Win64-Shipping.exe goto :no_unreal_version_selector
