[WinArm64] BootstrapPackagedGame changes to support arm64:
- it will install the correct VC redist.
- it is only built for x64 - expected to run via emulation on arm64 so there's only one bootstrap exe.

#rnx

diff --git a/Engine/Source/Programs/AutomationTool/Scripts/BuildProjectCommand.Automation.cs b/Engine/Source/Programs/AutomationTool/Scripts/BuildProjectCommand.Automation.cs
--- a/Engine/Source/Programs/AutomationTool/Scripts/BuildProjectCommand.Automation.cs
+++ b/Engine/Source/Programs/AutomationTool/Scripts/BuildProjectCommand.Automation.cs
@@ -197,7 +197,12 @@ namespace AutomationScripts
 				{
 					if (Params.ClientTargetPlatforms.Contains(new TargetPlatformDescriptor(BootstrapPackagedGamePlatformType)))
 					{
-						NoProjectAgenda.AddTarget("BootstrapPackagedGame", BootstrapPackagedGamePlatformType, UnrealBuildTool.UnrealTargetConfiguration.Shipping, InAddArgs: ClientBuildArgs);
+						string BootstrapBuildArgs = BootstrapPackagedGamePlatformType.IsInGroup(UnrealPlatformGroup.Windows) 
+							? Params.AdditionalBuildOptions // do not pass a specific architecture, so we default to x64
+							: ClientBuildArgs
+						;
+
+						NoProjectAgenda.AddTarget("BootstrapPackagedGame", BootstrapPackagedGamePlatformType, UnrealBuildTool.UnrealTargetConfiguration.Shipping, InAddArgs: BootstrapBuildArgs);
 					}
 				}
 			}
diff --git a/Engine/Source/Programs/Windows/BootstrapPackagedGame/BootstrapPackagedGame.Target.cs b/Engine/Source/Programs/Windows/BootstrapPackagedGame/BootstrapPackagedGame.Target.cs
--- a/Engine/Source/Programs/Windows/BootstrapPackagedGame/BootstrapPackagedGame.Target.cs
+++ b/Engine/Source/Programs/Windows/BootstrapPackagedGame/BootstrapPackagedGame.Target.cs
@@ -24,5 +24,11 @@ public class BootstrapPackagedGameTarget : TargetRules
 		bBuildWithEditorOnlyData = false;
 		bCompileAgainstEngine = false;
 		bCompileAgainstCoreUObject = false;
+
+		// enforce single unified x64 architecture for Windows
+		if (Target.Platform.IsInGroup(UnrealPlatformGroup.Windows) && (Target.Architectures.bIsMultiArch || Target.Architectures.SingleArchitecture != UnrealArch.X64))
+		{
+			throw new BuildException("BootstrapPackagedGame can only be built for x64. It is expected to run via emulation on Arm64");
+		}
 	}
 }
diff --git a/Engine/Source/Programs/Windows/BootstrapPackagedGame/Private/BootstrapPackagedGame.cpp b/Engine/Source/Programs/Windows/BootstrapPackagedGame/Private/BootstrapPackagedGame.cpp
--- a/Engine/Source/Programs/Windows/BootstrapPackagedGame/Private/BootstrapPackagedGame.cpp
+++ b/Engine/Source/Programs/Windows/BootstrapPackagedGame/Private/BootstrapPackagedGame.cpp
@@ -164,10 +164,32 @@ bool HasAppxPackagedVCRuntime()
 	return bHasVCLibs;
 }
 
+bool IsARM64HostPlatform()
+{
+	// try to find the IsWow64Process2 function
+	HMODULE hModule = GetModuleHandleW(TEXT("kernel32.dll"));
+	typedef BOOL( WINAPI *IsWow64Process2Proc)(HANDLE, USHORT*, USHORT*);
+	IsWow64Process2Proc fnIsWow64Process2 = hModule ? (IsWow64Process2Proc)GetProcAddress(hModule, "IsWow64Process2") : nullptr;
+
+	// query the native machine
+	if (fnIsWow64Process2 != nullptr)
+	{
+		USHORT ProcessMachine, NativeMachine;
+		if (fnIsWow64Process2( GetCurrentProcess(), &ProcessMachine, &NativeMachine))
+		{
+			return (NativeMachine == IMAGE_FILE_MACHINE_ARM64);
+		}
+	}
+
+	return false;
+}
+
 
 
 int InstallMissingPrerequisites(const WCHAR* BaseDirectory, const WCHAR* ExecDirectory)
 {
+	bool bIsARM64HostPlatform = IsARM64HostPlatform();
+
 	// Look for missing prerequisites
 	WCHAR MissingPrerequisites[1024] = { 0, };
 
@@ -195,7 +217,10 @@ int InstallMissingPrerequisites(const WCHAR* BaseDirectory, const WCHAR* ExecDir
 		if (bInstallVCRedist)
 		{
 			HKEY Hkey;
-			LSTATUS KeyOpenStatus = RegOpenKeyExW(HKEY_LOCAL_MACHINE, L"SOFTWARE\\Microsoft\\VisualStudio\\14.0\\VC\\Runtimes\\x64", 0, KEY_READ, &Hkey);
+			LSTATUS KeyOpenStatus = bIsARM64HostPlatform
+				? RegOpenKeyExW(HKEY_LOCAL_MACHINE, L"SOFTWARE\\Microsoft\\VisualStudio\\14.0\\VC\\Runtimes\\arm64", 0, KEY_READ, &Hkey)
+				: RegOpenKeyExW(HKEY_LOCAL_MACHINE, L"SOFTWARE\\Microsoft\\VisualStudio\\14.0\\VC\\Runtimes\\x64", 0, KEY_READ, &Hkey);
+			;
 
 			if (KeyOpenStatus == ERROR_SUCCESS)
 			{
@@ -239,7 +264,15 @@ int InstallMissingPrerequisites(const WCHAR* BaseDirectory, const WCHAR* ExecDir
 			}
 			if (bInstallVCRedist)
 			{
-				wcscat_s(MissingPrerequisites, TEXT("Microsoft Visual C++ 2015-2022 Redistributable (x64)\n"));
+				wcscat_s(MissingPrerequisites, TEXT("Microsoft Visual C++ 2015-2022 Redistributable "));
+				if (bIsARM64HostPlatform)
+				{
+					wcscat_s(MissingPrerequisites, TEXT("(arm64)\n"));
+				}
+				else
+				{
+					wcscat_s(MissingPrerequisites, TEXT("(x64)\n"));
+				}
 			}
 		}
 	}
@@ -253,7 +286,14 @@ int InstallMissingPrerequisites(const WCHAR* BaseDirectory, const WCHAR* ExecDir
 
 		// If we don't have the installer, just notify the user and quit
 		WCHAR PrereqInstaller[MAX_PATH];
-		PathCombine(PrereqInstaller, BaseDirectory, L"Engine\\Extras\\Redist\\en-us\\vc_redist.x64.exe");
+		if (bIsARM64HostPlatform)
+		{
+			PathCombine(PrereqInstaller, BaseDirectory, L"Engine\\Extras\\Redist\\en-us\\vc_redist.arm64.exe");
+		}
+		else
+		{
+			PathCombine(PrereqInstaller, BaseDirectory, L"Engine\\Extras\\Redist\\en-us\\vc_redist.x64.exe");
+		}
 		if(GetFileAttributes(PrereqInstaller) == INVALID_FILE_ATTRIBUTES)
 		{
 			MessageBox(NULL, MissingPrerequisitesMsg, NULL, MB_OK);
