Update vc redist to 14.44

#rnx

diff --git a/Engine/Source/Programs/Windows/BootstrapPackagedGame/Private/BootstrapPackagedGame.cpp b/Engine/Source/Programs/Windows/BootstrapPackagedGame/Private/BootstrapPackagedGame.cpp
--- a/Engine/Source/Programs/Windows/BootstrapPackagedGame/Private/BootstrapPackagedGame.cpp
+++ b/Engine/Source/Programs/Windows/BootstrapPackagedGame/Private/BootstrapPackagedGame.cpp
@@ -17,7 +17,7 @@ struct VersionInfo
 
 // This minimum should match the version installed by
 // Engine/Extras/Redist/en-us/vc_redist.x64.exe
-static const VersionInfo MinRedistVersion = { 14, 42, 34438, 0 };
+static const VersionInfo MinRedistVersion = { 14, 44, 35211, 0 };
 
 bool IsVersionValid(const VersionInfo& Version, const VersionInfo& MinVersion)
 {
