PR #13547: Add an ExecFileOverride resource to BootstrapPackagedGame

diff --git a/Engine/Source/Programs/AutomationTool/Win/WinPlatform.Automation.cs b/Engine/Source/Programs/AutomationTool/Win/WinPlatform.Automation.cs
--- a/Engine/Source/Programs/AutomationTool/Win/WinPlatform.Automation.cs
+++ b/Engine/Source/Programs/AutomationTool/Win/WinPlatform.Automation.cs
@@ -290,6 +290,8 @@ public class Win64Platform : Platform
 					
 					const int ExecArgsResourceId = 202;
 					Update.SetData(ExecArgsResourceId, ResourceType.RawData, Encoding.Unicode.GetBytes(StagedArguments + "\0"));
+					const int ExecFileOverrideResourceId = 205;
+					Update.SetData(ExecFileOverrideResourceId, ResourceType.RawData, Encoding.Unicode.GetBytes("\0"));
 				}
 			}
 			else
diff --git a/Engine/Source/Programs/Windows/BootstrapPackagedGame/Private/BootstrapPackagedGame.cpp b/Engine/Source/Programs/Windows/BootstrapPackagedGame/Private/BootstrapPackagedGame.cpp
--- a/Engine/Source/Programs/Windows/BootstrapPackagedGame/Private/BootstrapPackagedGame.cpp
+++ b/Engine/Source/Programs/Windows/BootstrapPackagedGame/Private/BootstrapPackagedGame.cpp
@@ -6,6 +6,7 @@
 #define IDI_EXEC_ARGS 202
 #define IDI_EXEC_FILE_ARM64 203
 #define IDI_EXEC_FILE_ARM64EC 204
+#define IDI_EXEC_FILE_OVERRIDE 205
 
 struct VersionInfo
 {
@@ -434,6 +435,16 @@ int WINAPI wWinMain(_In_ HINSTANCE hInstance, _In_opt_ HINSTANCE hPrevInstance,
 	delete[] TempExecDirectory;
 	PathRemoveFileSpec(ExecDirectory);
 
+	// Check for an exec file override.  The previously used exec file is used only to
+	// locate the binary folder for dependency checks, if an override is present.
+	WCHAR* ExecFileOverride = ReadResourceString(hInstance, MAKEINTRESOURCE(IDI_EXEC_FILE_OVERRIDE));
+	if (ExecFileOverride != NULL && wcslen(ExecFileOverride) > 0)
+	{
+		// If an override is present, use it instead of the one from the resource
+		delete[] ExecFile;
+		ExecFile = ExecFileOverride;
+	}
+
 	// Create a full command line for the program to run
 	WCHAR* BaseArgs = ReadResourceString(hInstance, MAKEINTRESOURCE(IDI_EXEC_ARGS));
 	size_t ChildCmdLineLength = wcslen(BaseDirectory) + wcslen(ExecFile) + wcslen(BaseArgs) + wcslen(CmdLine) + 20;
