Fix possible memory leak

#rnx

diff --git a/Engine/Source/Programs/Windows/BootstrapPackagedGame/Private/BootstrapPackagedGame.cpp b/Engine/Source/Programs/Windows/BootstrapPackagedGame/Private/BootstrapPackagedGame.cpp
--- a/Engine/Source/Programs/Windows/BootstrapPackagedGame/Private/BootstrapPackagedGame.cpp
+++ b/Engine/Source/Programs/Windows/BootstrapPackagedGame/Private/BootstrapPackagedGame.cpp
@@ -438,11 +438,18 @@ int WINAPI wWinMain(_In_ HINSTANCE hInstance, _In_opt_ HINSTANCE hPrevInstance,
 	// Check for an exec file override.  The previously used exec file is used only to
 	// locate the binary folder for dependency checks, if an override is present.
 	WCHAR* ExecFileOverride = ReadResourceString(hInstance, MAKEINTRESOURCE(IDI_EXEC_FILE_OVERRIDE));
-	if (ExecFileOverride != NULL && wcslen(ExecFileOverride) > 0)
+	if (ExecFileOverride != NULL)
 	{
-		// If an override is present, use it instead of the one from the resource
-		delete[] ExecFile;
-		ExecFile = ExecFileOverride;
+		if (wcslen(ExecFileOverride) > 0)
+		{
+			// If an override is present, use it instead of the one from the resource
+			delete[] ExecFile;
+			ExecFile = ExecFileOverride;
+		}
+		else
+		{
+			delete[] ExecFileOverride;
+		}
 	}
 
 	// Create a full command line for the program to run
@@ -463,12 +470,6 @@ int WINAPI wWinMain(_In_ HINSTANCE hInstance, _In_opt_ HINSTANCE hPrevInstance,
 
 	// Spawn the target executable
 	ExitCode = SpawnTarget(ChildCmdLine);
-	if(ExitCode != 0)
-	{
-		delete[] ChildCmdLine;
-		return ExitCode;
-	}
-
 	delete[] ChildCmdLine;
 	return ExitCode;
 }
