[Game Input] Add the GameInputRedist.msi as an optional BootstrapPackagedGame pre-req step.

In order to include this new GameInputRedist.msi file with packaging, users must do the following:

Set [GameInput]:IncludeRedistFiles=True in the project's DefaultEngine.ini. Like this:

; Include the GameInput redist files on windows targets
[GameInput]
IncludeRedistFiles=True

This will ensure that the redist file is staged with your packaged game when the "-Prereqs" flag is set in UAT.

[RN] You can now set [GameInput]:IncludeRedistFiles=True in your project's DefaultEngine.ini to have BootstrapPackagedGame run a GameInputRedist.msi installer! You should set this for any packaged game using the Game Input for Windows plugin.

diff --git a/Engine/Source/Programs/AutomationTool/Win/WinPlatform.Automation.cs b/Engine/Source/Programs/AutomationTool/Win/WinPlatform.Automation.cs
--- a/Engine/Source/Programs/AutomationTool/Win/WinPlatform.Automation.cs
+++ b/Engine/Source/Programs/AutomationTool/Win/WinPlatform.Automation.cs
@@ -200,6 +200,8 @@ public class Win64Platform : Platform
 			//SC.StageFile(StagedFileType.NonUFS, FileReference.Combine(SC.EngineRoot, "Extras", "Redist", "en-us", "UEPrereqSetup_x64.exe"));
 			SC.StageFile(StagedFileType.NonUFS, FileReference.Combine(SC.EngineRoot, "Extras", "Redist", "en-us", "vc_redist.x64.exe"));
 			SC.StageFile(StagedFileType.NonUFS, FileReference.Combine(SC.EngineRoot, "Extras", "Redist", "en-us", "vc_redist.arm64.exe")); // always staging arm64 redist in case we are run on an arm64 device in emulation mode
+
+			StageGameInputWindowsRedist(SC);
 		}
 
 		if (!string.IsNullOrWhiteSpace(Params.AppLocalDirectory))
@@ -211,6 +213,34 @@ public class Win64Platform : Platform
 		}
 	}
 
+	/**
+	 * Returns true if the project has set the config setting to include the GameInputRedist with its packaging
+	 */
+	protected bool ShouldStageGameInputWindowsRedist(DeploymentContext SC)
+	{
+		// Only include the Game Input redist file if the project has opted into via the Engine config setting for their project
+		ConfigHierarchy EngineConfig = ConfigCache.ReadHierarchy(ConfigHierarchyType.Engine, DirectoryReference.FromFile(SC.RawProjectPath), UnrealTargetPlatform.Win64);
+		
+		//		; In YourGameProject/Config/DefaultEngine.ini...
+		//		[GameInput]
+		//		IncludeRedistFiles=True
+		
+		EngineConfig.GetBool("GameInput", "IncludeRedistFiles", out bool bShouldIncludeRedistFiles);
+
+		return bShouldIncludeRedistFiles;
+	}
+	
+	protected void StageGameInputWindowsRedist(DeploymentContext SC)
+	{
+		if (ShouldStageGameInputWindowsRedist(SC))
+		{
+			Console.WriteLine("Including Game Input prereq installer!");
+			
+			// Engine/Extras/Redist/en-us/GameInputRedist.msi
+			SC.StageFile(StagedFileType.NonUFS, FileReference.Combine(SC.EngineRoot, "Extras", "Redist", "en-us", "GameInputRedist.msi"));	
+		}
+	}
+
 	public override void ExtractPackage(ProjectParams Params, string SourcePath, string DestinationPath)
     {
     }
@@ -292,6 +322,14 @@ public class Win64Platform : Platform
 					Update.SetData(ExecArgsResourceId, ResourceType.RawData, Encoding.Unicode.GetBytes(StagedArguments + "\0"));
 					const int ExecFileOverrideResourceId = 205;
 					Update.SetData(ExecFileOverrideResourceId, ResourceType.RawData, Encoding.Unicode.GetBytes("\0"));
+
+					// Set a flag on the exec file so that the BootstrapPackagedGame code can check if we should run the GameInputRedistInstaller
+					if (ShouldStageGameInputWindowsRedist(SC))
+					{
+						const int ExecFileIncludeGameInputRedist = 206;
+						Update.SetData(ExecFileIncludeGameInputRedist, ResourceType.RawData, Encoding.Unicode.GetBytes("\0"));	
+						Logger.LogInformation("Including the GameInputRedist.msi file with the packaged project");
+					}
 				}
 			}
 			else
diff --git a/Engine/Source/Programs/Windows/BootstrapPackagedGame/Private/BootstrapPackagedGame.cpp b/Engine/Source/Programs/Windows/BootstrapPackagedGame/Private/BootstrapPackagedGame.cpp
--- a/Engine/Source/Programs/Windows/BootstrapPackagedGame/Private/BootstrapPackagedGame.cpp
+++ b/Engine/Source/Programs/Windows/BootstrapPackagedGame/Private/BootstrapPackagedGame.cpp
@@ -7,6 +7,7 @@
 #define IDI_EXEC_FILE_ARM64 203
 #define IDI_EXEC_FILE_ARM64EC 204
 #define IDI_EXEC_FILE_OVERRIDE 205
+#define IDI_EXEC_FILE_INCLUDE_GAME_INPUT_REDIST 206
 
 struct VersionInfo
 {
@@ -19,6 +20,8 @@ struct VersionInfo
 // This minimum should match the version installed by
 // Engine/Extras/Redist/en-us/vc_redist.x64.exe
 static const VersionInfo MinRedistVersion = { 14, 44, 35211, 0 };
+static constexpr VersionInfo MinGameInputVersion = { 0, 2309, 26100, 5074 };
+static constexpr VersionInfo MinGameInputRedistVersion = { 3, 0, 26100, 6154 };
 
 bool IsVersionValid(const VersionInfo& Version, const VersionInfo& MinVersion)
 {
@@ -187,7 +190,134 @@ bool IsARM64HostPlatform()
 	return false;
 }
 
+/**
+ * Returns true if the program requires the GameInputRedist file to be included. 
+ */
+bool DoesProjectRequireGameInputRedist(HINSTANCE hInstance)
+{
+	return ReadResourceString(hInstance, MAKEINTRESOURCE(IDI_EXEC_FILE_INCLUDE_GAME_INPUT_REDIST)) != nullptr;
+}
+
+/**
+* Returns true if we have a valid version of Game Input installed already on the machine.
+*/
+bool HasGameInputInstalled()
+{
+	// The min version should match the version installed by
+	// Engine/Extras/Redist/en-us/GameInputRedist.msi
+	
+	// We have two DLLs to version check for game input, GameInput and GameInputRedist.
+	bool bHasValidGameInput = false;
+	bool bHasValidGameInputRedist = false;
+
+	// Check the System32 folder for if the GameInput dll is available
+	WCHAR SystemRoot[MAX_PATH] = { 0 };
+	GetEnvironmentVariable(L"SystemRoot", SystemRoot, MAX_PATH);
+	WCHAR System32Path[MAX_PATH] = { 0 };
+	PathCombine(System32Path, SystemRoot, L"system32");
+
+	// Check GameInput.dll verison 0.2309.26100.5074
+	{
+		if (IsDllValid(System32Path, L"GameInput.dll", MinGameInputVersion))
+		{
+			bHasValidGameInput = true;
+		}	
+		else if (IsDllValid(L"GameInput.dll", MinGameInputVersion))
+		{
+			bHasValidGameInput = true;
+		}
+	}
 
+	// Check GameInputRedist.dll version 3.0.26100.6154
+	{
+		if (IsDllValid(System32Path, L"GameInputRedist.dll", MinGameInputRedistVersion))
+		{
+			bHasValidGameInputRedist = true;
+		}	
+		else if (IsDllValid(L"GameInputRedist.dll", MinGameInputRedistVersion))
+		{
+			bHasValidGameInputRedist = true;
+		}
+	}
+
+	return bHasValidGameInput && bHasValidGameInputRedist;
+}
+
+/**
+* Runs the GameInputRedist.msi to install the required GameInput DLLs and service.
+* 
+* This should only be called if you have opted in to have your project use the GameInput library for processing gamepad input
+* 
+* @return 0 if successful install or GameInput is already installed. Error code if anything failed.
+*/
+int InstallGameInputPrereq(const WCHAR* BaseDirectory)
+{
+	static constexpr int SuccessfulExitCode = 0;
+
+	const bool bHasGameInputInstalled = HasGameInputInstalled();
+	
+	// We already have a valid GameInput version installed, we can exit with success
+	if (bHasGameInputInstalled)
+	{
+		wprintf(L"[GameInput installer] A valid version of GameInput is already installed.");
+		return SuccessfulExitCode;
+	}
+
+	// Engine/Extras/Redist/en-us/GameInputRedist.msi
+	// Ensure that we actually have the installer file available for us. If not, error out
+	WCHAR GameInputPrereqInstaller[1024] = { 0, };
+	PathCombine(GameInputPrereqInstaller, BaseDirectory, L"Engine\\Extras\\Redist\\en-us\\GameInputRedist.msi");
+
+	if (GetFileAttributes(GameInputPrereqInstaller) == INVALID_FILE_ATTRIBUTES)
+	{
+		WCHAR MissingGameInputInstallerMsg[1024];
+		wsprintf(MissingGameInputInstallerMsg, L"The Game Input installer (GameInputRedist.msi) is missing, and is required to run this program.\n\n");
+
+		MessageBox(NULL, MissingGameInputInstallerMsg, NULL, MB_OK);
+		return 9010;
+	}
+
+	// Prompt the user to install GameInput now
+	WCHAR InstallGameInputPromptMessage[1024];
+	wsprintf(InstallGameInputPromptMessage, L"The Game Input Library is missing, and is required to run this program.\n\nWould you like to install it now?");
+
+	if (MessageBox(NULL, InstallGameInputPromptMessage, NULL, MB_YESNO) == IDNO)
+	{
+		// The user has elected not to install GameInput. Exit.
+		return 9011;
+	}
+
+	wprintf(L"[GameInput installer] Starting... (%ls)", GameInputPrereqInstaller);
+
+	// Start the game input installer
+	SHELLEXECUTEINFO ShellExecuteInfo;
+	ZeroMemory(&ShellExecuteInfo, sizeof(ShellExecuteInfo));
+	ShellExecuteInfo.cbSize = sizeof(ShellExecuteInfo);
+	ShellExecuteInfo.fMask = SEE_MASK_NOCLOSEPROCESS;
+	ShellExecuteInfo.nShow = SW_SHOWNORMAL;
+	ShellExecuteInfo.lpFile = GameInputPrereqInstaller;
+	if (!ShellExecuteExW(&ShellExecuteInfo))
+	{
+		wprintf(L"[GameInput installer] Failed to start GameInputRedist.msi");
+		return 9012;
+	}
+
+	// Wait for the process to complete, then get its exit code
+	DWORD ExitCode = 0;
+	WaitForSingleObject(ShellExecuteInfo.hProcess, INFINITE);
+
+	GetExitCodeProcess(ShellExecuteInfo.hProcess, &ExitCode);
+	CloseHandle(ShellExecuteInfo.hProcess);
+
+	if (ExitCode != 0)
+	{		
+		wprintf(L"[GameInput installer] Failed! Error code %lu", ExitCode);
+		return 9013;
+	}	
+
+	// Successful install
+	return SuccessfulExitCode;
+}
 
 int InstallMissingPrerequisites(const WCHAR* BaseDirectory, const WCHAR* ExecDirectory)
 {
@@ -468,6 +598,18 @@ int WINAPI wWinMain(_In_ HINSTANCE hInstance, _In_opt_ HINSTANCE hPrevInstance,
 		return ExitCode;
 	}
 
+	// Projects with the "[GameInput]:IncludeRedistFiles=True" config value set when packaged will include
+	// the game input redist file.
+	if (DoesProjectRequireGameInputRedist(hInstance))
+	{		
+		int GameInputRedistExitCode = InstallGameInputPrereq(BaseDirectory);
+		if (GameInputRedistExitCode != 0)
+		{
+			delete[] ChildCmdLine;
+			return ExitCode;
+		}
+	}
+
 	// Spawn the target executable
 	ExitCode = SpawnTarget(ChildCmdLine);
 	delete[] ChildCmdLine;
