[GameInput] Fix incorrect return value. Ensure that we return the GameInput installer return value and not the previous one. Also, make sure to account for the MSI installation code 1638 for any future versions of the GameInputRedist which may be installed on the users machine. Fix possible memory leak in DoesProjectRequireGameInputRedist

#rnx

diff --git a/Engine/Source/Programs/Windows/BootstrapPackagedGame/Private/BootstrapPackagedGame.cpp b/Engine/Source/Programs/Windows/BootstrapPackagedGame/Private/BootstrapPackagedGame.cpp
--- a/Engine/Source/Programs/Windows/BootstrapPackagedGame/Private/BootstrapPackagedGame.cpp
+++ b/Engine/Source/Programs/Windows/BootstrapPackagedGame/Private/BootstrapPackagedGame.cpp
@@ -195,7 +195,13 @@ bool IsARM64HostPlatform()
  */
 bool DoesProjectRequireGameInputRedist(HINSTANCE hInstance)
 {
-	return ReadResourceString(hInstance, MAKEINTRESOURCE(IDI_EXEC_FILE_INCLUDE_GAME_INPUT_REDIST)) != nullptr;
+	const WCHAR* ResourceString = ReadResourceString(hInstance, MAKEINTRESOURCE(IDI_EXEC_FILE_INCLUDE_GAME_INPUT_REDIST));
+
+	const bool bHasGameInputRedistResource = ResourceString != nullptr;
+	
+	delete [] ResourceString;
+	
+	return bHasGameInputRedistResource;
 }
 
 /**
@@ -308,8 +314,9 @@ int InstallGameInputPrereq(const WCHAR* BaseDirectory)
 
 	GetExitCodeProcess(ShellExecuteInfo.hProcess, &ExitCode);
 	CloseHandle(ShellExecuteInfo.hProcess);
-
-	if (ExitCode != 0)
+	
+	// 1638: Newer version already installed
+	if (ExitCode != 0 && ExitCode != 1638)
 	{		
 		wprintf(L"[GameInput installer] Failed! Error code %lu", ExitCode);
 		return 9013;
@@ -606,7 +613,7 @@ int WINAPI wWinMain(_In_ HINSTANCE hInstance, _In_opt_ HINSTANCE hPrevInstance,
 		if (GameInputRedistExitCode != 0)
 		{
 			delete[] ChildCmdLine;
-			return ExitCode;
+			return GameInputRedistExitCode;
 		}
 	}
 
