BootstrapPackagedGame: Fix app-local dll check

#rnx
#Synced-CL 51561575

diff --git a/Engine/Source/Programs/Windows/BootstrapPackagedGame/Private/BootstrapPackagedGame.cpp b/Engine/Source/Programs/Windows/BootstrapPackagedGame/Private/BootstrapPackagedGame.cpp
--- a/Engine/Source/Programs/Windows/BootstrapPackagedGame/Private/BootstrapPackagedGame.cpp
+++ b/Engine/Source/Programs/Windows/BootstrapPackagedGame/Private/BootstrapPackagedGame.cpp
@@ -40,7 +40,7 @@ bool TryLoadDll(const WCHAR* ExecDirectory, const WCHAR* Name)
 	{
 		return false;
 	}
-	HMODULE Handle = LoadLibrary(AppLocalPath);
+	HMODULE Handle = LoadLibraryEx(AppLocalPath, nullptr, LOAD_LIBRARY_SEARCH_DLL_LOAD_DIR | LOAD_LIBRARY_SEARCH_DEFAULT_DIRS);
 	if (Handle != nullptr)
 	{
 		FreeLibrary(Handle);
