Add ASDTool - Advanced Shader Delivery pipeline tool

Engine program with 3 commands for the D3D12 Advanced Shader Delivery pipeline:
- generatesodb: Generate SODB from cooked shader data (SHK + ushaderbytecode)
- compilepsdbs: Compile SODBs to PSDBs via IHV compiler plugins + D3D12StateObjectCompiler.exe
- registerpsdbs: Register PSDBs with D3D12 shader cache runtime for precompiled shader delivery
Supports local adapter detection, Kraken compression, and configurable app identity.

#Synced-CL 52245732

diff --git a/Engine/Source/Programs/Windows/ASDTool/ASDTool.Build.cs b/Engine/Source/Programs/Windows/ASDTool/ASDTool.Build.cs
new file mode 100644
--- /dev/null
+++ b/Engine/Source/Programs/Windows/ASDTool/ASDTool.Build.cs
@@ -0,0 +1,43 @@
+// Copyright Epic Games, Inc. All Rights Reserved.
+
+using UnrealBuildTool;
+
+public class ASDTool : ModuleRules
+{
+	public ASDTool(ReadOnlyTargetRules Target) : base(Target)
+	{
+		bTreatAsEngineModule = true;
+
+		PrivateIncludePathModuleNames.AddRange(new string[] {
+			"Launch",
+			"TargetPlatform"
+		});
+
+		PrivateDependencyModuleNames.AddRange(new string[] {
+			"Core",
+			"CoreUObject",
+			"ApplicationCore",
+			"Projects",
+			"ShaderCompilerCommon",
+			"RenderCore",
+			"RHI",
+			"RHICore",
+			"D3D12RHI",
+			"PakFile",
+			"PakFileUtilities",
+			"SQLiteCore"
+		});
+
+		bEnableExceptions = true;
+
+		if (Target.Platform == UnrealTargetPlatform.Win64)
+		{
+			AddEngineThirdPartyPrivateStaticDependencies(Target, "AgilitySDK");
+			PublicSystemLibraries.AddRange(new string[] { "d3d12.lib", "dxgi.lib" });
+
+			// Expose the AgilitySDK version string as a preprocessor define so runtime code
+			// can locate compiler binaries without hardcoding the version.
+			PrivateDefinitions.Add($"AGILITY_SDK_VERSION_STRING=\"{AgilitySDK.DefaultVersion}\"");
+		}
+	}
+}
diff --git a/Engine/Source/Programs/Windows/ASDTool/ASDTool.Target.cs b/Engine/Source/Programs/Windows/ASDTool/ASDTool.Target.cs
new file mode 100644
--- /dev/null
+++ b/Engine/Source/Programs/Windows/ASDTool/ASDTool.Target.cs
@@ -0,0 +1,45 @@
+// Copyright Epic Games, Inc. All Rights Reserved.
+
+using UnrealBuildTool;
+
+[SupportedPlatforms(UnrealPlatformClass.Editor)]
+public class ASDToolTarget : TargetRules
+{
+	public ASDToolTarget(TargetInfo Target) : base(Target)
+	{
+		Type = TargetType.Program;
+		LinkType = TargetLinkType.Monolithic;
+		IncludeOrderVersion = EngineIncludeOrderVersion.Latest;
+		LaunchModuleName = "ASDTool";
+
+		bBuildDeveloperTools = false;
+
+		bool bDebugOrDevelopment = Target.Configuration == UnrealTargetConfiguration.Debug || Target.Configuration == UnrealTargetConfiguration.Development;
+		bBuildWithEditorOnlyData = Target.Platform.IsInGroup(UnrealPlatformGroup.Desktop) && bDebugOrDevelopment;
+
+		bCompileAgainstEngine = false;
+		bCompileAgainstCoreUObject = false;
+		bCompileAgainstApplicationCore = true;
+		bCompileICU = false;
+
+		bIsBuildingConsoleApplication = true;
+
+		bForceCompileDevelopmentAutomationTests = false;
+		bCompileWithStatsWithoutEngine = false;
+
+		bUseLoggingInShipping = true;
+		bLegalToDistributeBinary = true;
+
+		WindowsPlatform.bUseBundledDbgHelp = false;
+		WindowsPlatform.bPixProfilingEnabled = false;
+		WindowsPlatform.bMergeIdenticalCOMDATs = true;
+		WindowsPlatform.bStripUnreferencedSymbols = true;
+
+		// SQLite needs system malloc (not UE custom allocator) for standalone tools
+		bCompileCustomSQLitePlatform = false;
+		GlobalDefinitions.Add("SQLITE_OS_OTHER=0");
+		GlobalDefinitions.Add("SQLITE_OS_WIN=1");
+
+		GlobalDefinitions.Add("AUTOSDKS_ENABLED=0");
+	}
+}
diff --git a/Engine/Source/Programs/Windows/ASDTool/Private/ASDTool.cpp b/Engine/Source/Programs/Windows/ASDTool/Private/ASDTool.cpp
new file mode 100644
--- /dev/null
+++ b/Engine/Source/Programs/Windows/ASDTool/Private/ASDTool.cpp
@@ -0,0 +1,87 @@
+// Copyright Epic Games, Inc. All Rights Reserved.
+
+#include "RequiredProgramMainCPPInclude.h"
+
+#include "Misc/CommandLine.h"
+#include "Misc/Parse.h"
+#include "Misc/CString.h"
+#include "HAL/PlatformTime.h"
+
+#include "ASDToolCommands.h"
+
+DEFINE_LOG_CATEGORY(LogASDTool);
+
+// These macros are not properly defined by UBT for engine programs
+#define IMPLEMENT_ENCRYPTION_KEY_REGISTRATION()
+#define IMPLEMENT_SIGNING_KEY_REGISTRATION()
+IMPLEMENT_APPLICATION(ASDTool, "ASDTool");
+
+int32 ASDMain()
+{
+	const TCHAR* CmdLine = FCommandLine::Get();
+
+	// Parse command - first argument after the exe
+	TArray<FString> Tokens;
+	TArray<FString> Switches;
+	FCommandLine::Parse(CmdLine, Tokens, Switches);
+
+	if (Tokens.Num() < 1)
+	{
+		UE_LOG(LogASDTool, Display, TEXT("ASDTool - Advanced Shader Delivery Tool"));
+		UE_LOG(LogASDTool, Display, TEXT(""));
+		UE_LOG(LogASDTool, Display, TEXT("Usage: ASDTool <command> [options]"));
+		UE_LOG(LogASDTool, Display, TEXT(""));
+		UE_LOG(LogASDTool, Display, TEXT("Commands:"));
+		UE_LOG(LogASDTool, Display, TEXT("  generatesodb    Generate State Object Database (.sodb) for compute and graphics shaders"));
+		UE_LOG(LogASDTool, Display, TEXT("  compilepsdbs    Compile SODBs into Precompiled Shader Databases (.psdb)"));
+		UE_LOG(LogASDTool, Display, TEXT("  registerpsdbs   Register PSDBs with the D3D12 shader cache runtime"));
+		UE_LOG(LogASDTool, Display, TEXT(""));
+		UE_LOG(LogASDTool, Display, TEXT("Use 'ASDTool <command> -help' for more information on a specific command."));
+		return 1;
+	}
+
+	const TCHAR* Command = *Tokens[0];
+	int32 RetCode = 0;
+
+	if (!FCString::Stricmp(Command, TEXT("generatesodb")))
+	{
+		RetCode = ASDTool::GenerateSODBs();
+	}
+	else if (!FCString::Stricmp(Command, TEXT("compilepsdbs")))
+	{
+		RetCode = ASDTool::CompilePSDBs();
+	}
+	else if (!FCString::Stricmp(Command, TEXT("registerpsdbs")))
+	{
+		RetCode = ASDTool::RegisterPSDBs();
+	}
+	else
+	{
+		UE_LOG(LogASDTool, Error, TEXT("Unknown command: %s"), Command);
+		RetCode = 1;
+	}
+
+	return RetCode;
+}
+
+INT32_MAIN_INT32_ARGC_TCHAR_ARGV()
+{
+	FTaskTagScope Scope(ETaskTag::EGameThread);
+	ON_SCOPE_EXIT
+	{
+		LLM(FLowLevelMemTracker::Get().ProcessCommandLine(FCommandLine::Get()));
+		RequestEngineExit(TEXT("Exiting"));
+		FEngineLoop::AppPreExit();
+		FModuleManager::Get().UnloadModulesAtShutdown();
+		FEngineLoop::AppExit();
+	};
+
+	GEngineLoop.PreInit(ArgC, ArgV);
+	double StartTime = FPlatformTime::Seconds();
+
+	int32 RetCode = ASDMain();
+
+	UE_LOG(LogASDTool, Display, TEXT("Total time: %.1f seconds"), FPlatformTime::Seconds() - StartTime);
+
+	return RetCode;
+}
diff --git a/Engine/Source/Programs/Windows/ASDTool/Private/ASDToolCommands.h b/Engine/Source/Programs/Windows/ASDTool/Private/ASDToolCommands.h
new file mode 100644
--- /dev/null
+++ b/Engine/Source/Programs/Windows/ASDTool/Private/ASDToolCommands.h
@@ -0,0 +1,34 @@
+// Copyright Epic Games, Inc. All Rights Reserved.
+
+#pragma once
+
+#include "CoreMinimal.h"
+
+/**
+ * ASDTool - Advanced Shader Delivery Tool
+ *
+ * Commands:
+ *   generatesodb   - Generate State Object Database (.sodb) from cooked shader data
+ *   compilepsdbs   - Compile SODBs into Precompiled Shader Databases (.psdb)
+ *   registerpsdbs  - Register PSDBs with the D3D12 shader cache runtime
+ */
+
+DECLARE_LOG_CATEGORY_EXTERN(LogASDTool, Log, All);
+
+namespace ASDTool
+{
+
+// -- Subcommand entry points -------------------------------------------------
+
+/** Generate State Object Database (.sodb) from cooked shader data. */
+int32 GenerateSODBs();
+
+/** Compile SODBs into Precompiled Shader Databases (.psdb) via IHV compiler plugins. */
+int32 CompilePSDBs();
+
+/** Register PSDBs with the D3D12 shader cache runtime. */
+int32 RegisterPSDBs();
+
+
+
+};
diff --git a/Engine/Source/Programs/Windows/ASDTool/Private/ASDToolShaderUtils.cpp b/Engine/Source/Programs/Windows/ASDTool/Private/ASDToolShaderUtils.cpp
new file mode 100644
--- /dev/null
+++ b/Engine/Source/Programs/Windows/ASDTool/Private/ASDToolShaderUtils.cpp
@@ -0,0 +1,266 @@
+// Copyright Epic Games, Inc. All Rights Reserved.
+
+#include "ASDToolShaderUtils.h"
+#include "ASDToolCommands.h"
+
+#include "Async/ParallelFor.h"
+#include "HAL/FileManager.h"
+#include "Serialization/LargeMemoryReader.h"
+#include "ShaderCompilerCore.h"
+
+
+namespace ASDTool
+{
+
+//--------------------------------------------------------------------------------------------------
+// Internal: decompress a shader archive entry into its component parts.
+//--------------------------------------------------------------------------------------------------
+static void DecompressAndExtract(
+	const FShaderCodeEntry& Entry,
+	const uint8* CompressedPtr,
+	TArray<uint8>& UncompressedBuffer,
+	TArrayView64<const uint8>& OutShaderCode,
+	TArrayView64<const uint8>& OutSRTData,
+	TArrayView64<const uint8>& OutOptionalData)
+{
+	UncompressedBuffer.SetNumUninitialized(Entry.GetUncompressedSize());
+
+	if (Entry.GetUncompressedSize() != Entry.Size)
+	{
+		ShaderCodeArchive::DecompressShaderWithOodle(
+			UncompressedBuffer.GetData(), Entry.GetUncompressedSize(),
+			CompressedPtr, Entry.Size);
+	}
+	else
+	{
+		FMemory::Memcpy(UncompressedBuffer.GetData(), CompressedPtr, Entry.GetUncompressedSize());
+	}
+
+	// Parse past the SRT to get to the shader code
+	FLargeMemoryReader ShaderReader(UncompressedBuffer.GetData(), UncompressedBuffer.Num());
+	FShaderResourceTable SRT;
+	ShaderReader << SRT;
+
+	const uint32 ShaderCodeStart = (uint32)ShaderReader.Tell();
+	uint32 ShaderCodeSize = Entry.GetUncompressedSize() - ShaderCodeStart;
+	
+	// We can have some optional data appended after the shader code. The size of this optional data is stored in the last 4 bytes.
+	uint32 OptionalDataSize = 0;
+	if (ShaderCodeSize > sizeof(uint32))
+	{
+		const uint8* End = UncompressedBuffer.GetData() + UncompressedBuffer.Num();
+		OptionalDataSize = ((const unaligned_uint32*)End)[-1];
+		OptionalDataSize = FMath::Min(OptionalDataSize, ShaderCodeSize);
+		ShaderCodeSize -= OptionalDataSize;
+	}
+
+	OutSRTData      = TArrayView64<const uint8>(UncompressedBuffer.GetData(), ShaderCodeStart);
+	OutShaderCode   = TArrayView64<const uint8>(UncompressedBuffer.GetData() + ShaderCodeStart, ShaderCodeSize);
+	OutOptionalData = TArrayView64<const uint8>(UncompressedBuffer.GetData() + ShaderCodeStart + ShaderCodeSize, OptionalDataSize);
+}
+
+//--------------------------------------------------------------------------------------------------
+// ProcessArchive
+//--------------------------------------------------------------------------------------------------
+
+int32 ProcessArchive(
+	const FString& ArchivePath,
+	const TArray<FShaderKeyEntry>& ShaderKeys,
+	TSet<FShaderHash>& AlreadySeen,
+	FShaderVisitor Visitor,
+	bool bSingleThreaded,
+	int32* OutNumNotFound,
+	int32* OutNumDuplicates)
+{
+	UE_LOG(LogASDTool, Display, TEXT("Processing archive '%s'..."), *ArchivePath);
+
+	TUniquePtr<FArchive> LibraryAr(IFileManager::Get().CreateFileReader(*ArchivePath));
+	if (!LibraryAr)
+	{
+		UE_LOG(LogASDTool, Error, TEXT("Failed to open '%s'."), *ArchivePath);
+		return 0;
+	}
+
+	FIoBuffer LibraryBuffer(LibraryAr->TotalSize());
+	LibraryAr->Serialize(LibraryBuffer.GetData(), LibraryBuffer.GetSize());
+	FLargeMemoryReader LibraryReader(LibraryBuffer.GetData(), LibraryBuffer.GetSize());
+
+	uint32 Version = 0;
+	LibraryReader << Version;
+
+	FSerializedShaderArchive SerializedShaders;
+	LibraryReader << SerializedShaders;
+	const int64 CodeOffset = LibraryReader.Tell();
+	const uint8* CodeStart = LibraryBuffer.GetData() + CodeOffset;
+
+	const int32 NumShaders = SerializedShaders.GetNumShaders();
+	const TArrayView<const FShaderCodeEntry> ShaderEntries = SerializedShaders.GetShaderEntries();
+	const TArrayView<const FShaderHash> ShaderHashes = SerializedShaders.GetShaderHashes();
+
+	// Sort indices by hash for merge-join with ShaderKeys
+	TArray<int32> SortedIndices;
+	SortedIndices.SetNum(NumShaders);
+	for (int32 j = 0; j < NumShaders; ++j) { SortedIndices[j] = j; }
+	Algo::Sort(SortedIndices, [ShaderHashes](int A, int B)
+	{
+		return ShaderHashes[A] < ShaderHashes[B];
+	});
+
+	struct FExtractJob
+	{
+		const FShaderCodeEntry* Entry;
+		const uint8* ShaderPtr;
+		int32 KeyIndex;
+	};
+
+	TArray<FExtractJob> Jobs;
+	Jobs.Reserve(NumShaders);
+
+	int i = 0, KeyIdx = 0;
+	int NumNotFound = 0, NumDuplicates = 0;
+
+	while (i < NumShaders && KeyIdx < ShaderKeys.Num())
+	{
+		const int ShaderIdx = SortedIndices[i];
+		const FShaderCodeEntry& Entry = ShaderEntries[ShaderIdx];
+		const FShaderHash& Hash = ShaderHashes[ShaderIdx];
+
+		if (Hash.Hash < ShaderKeys[KeyIdx].Hash.Hash)
+		{
+			++NumNotFound;
+			++i;
+			continue;
+		}
+		if (Hash.Hash > ShaderKeys[KeyIdx].Hash.Hash)
+		{
+			++KeyIdx;
+			continue;
+		}
+
+		// Match found - check deduplication
+		const bool bAlreadySeen = AlreadySeen.Contains(Hash);
+		if (!bAlreadySeen)
+		{
+			AlreadySeen.Add(Hash);
+		}
+
+		++KeyIdx;
+		++i;
+
+		if (bAlreadySeen)
+		{
+			++NumDuplicates;
+			continue;
+		}
+
+		FExtractJob& Job = Jobs.Emplace_GetRef();
+		Job.Entry    = &Entry;
+		Job.ShaderPtr = CodeStart + Entry.Offset;
+		Job.KeyIndex  = KeyIdx - 1;  // KeyIdx was already incremented
+	}
+
+	NumNotFound += (NumShaders - i);
+
+	if (OutNumNotFound)
+	{
+		*OutNumNotFound = NumNotFound;
+	}
+	if (OutNumDuplicates)
+	{
+		*OutNumDuplicates = NumDuplicates;
+	}
+
+	TAtomic<int32> NumProcessed{0};
+
+	ParallelFor(Jobs.Num(), [&Jobs, &ShaderKeys, &Visitor, &NumProcessed](int JobIdx)
+	{
+		const FExtractJob& Job = Jobs[JobIdx];
+		const FShaderKeyEntry& KeyInfo = ShaderKeys[Job.KeyIndex];
+
+		TArray<uint8> UncompressedBuffer;
+		TArrayView64<const uint8> ShaderCode, SRTData, OptionalData;
+		DecompressAndExtract(*Job.Entry, Job.ShaderPtr, UncompressedBuffer, ShaderCode, SRTData, OptionalData);
+
+		FShaderArchiveEntry ArchiveEntry;
+		ArchiveEntry.Hash             = KeyInfo.Hash;
+		ArchiveEntry.Frequency        = Job.Entry->GetFrequency();
+		ArchiveEntry.ShaderCode        = ShaderCode;
+		ArchiveEntry.OptionalData      = OptionalData;
+		ArchiveEntry.UncompressedBuffer = &UncompressedBuffer;
+
+		Visitor(ArchiveEntry);
+		++NumProcessed;
+
+	}, bSingleThreaded);
+
+	UE_LOG(LogASDTool, Display, TEXT("Processed %d unique shaders from '%s'. %d not found in SHK, %d duplicates skipped."),
+		(int32)NumProcessed, *ArchivePath, NumNotFound, NumDuplicates);
+
+	return (int32)NumProcessed;
+}
+
+//--------------------------------------------------------------------------------------------------
+// Optional data helpers
+//--------------------------------------------------------------------------------------------------
+
+const uint8* FindOptionalData(
+	TArrayView64<const uint8> OptionalData,
+	EShaderOptionalDataKey InKey,
+	uint32 ValueSize)
+{
+	if (OptionalData.Num() == 0)
+	{
+		return nullptr;
+	}
+
+	const uint8* Current = OptionalData.GetData();
+	const uint8* End = Current + OptionalData.Num();
+
+	// OptionalData is trimmed to exactly OptionalDataSize bytes by DecompressAndExtract -
+	// the trailing size sentinel is not included. No adjustment needed.
+
+	while (Current < End)
+	{
+		const EShaderOptionalDataKey Key = EShaderOptionalDataKey(*Current++);
+		if (Current + sizeof(uint32) > End)
+		{
+			break;
+		}
+		uint32 Size = 0;
+		FMemory::Memcpy(&Size, Current, sizeof(uint32));
+		Current += sizeof(Size);
+
+		if (Key == InKey && Size == ValueSize)
+		{
+			if (Current + Size > End)
+			{
+				break;
+			}
+			return Current;
+		}
+
+		// Bounds-check before advancing to avoid UB from out-of-range pointer arithmetic
+		if (Size > (uint32)(End - Current))
+		{
+			break;
+		}
+		Current += Size;
+	}
+
+	return nullptr;
+}
+
+bool ExtractResourceCounts(
+	TArrayView64<const uint8> OptionalData,
+	FShaderCodePackedResourceCounts& OutCounts)
+{
+	const uint8* Data = FindOptionalData(OptionalData, FShaderCodePackedResourceCounts::Key, sizeof(FShaderCodePackedResourceCounts));
+	if (Data)
+	{
+		FMemory::Memcpy(&OutCounts, Data, sizeof(FShaderCodePackedResourceCounts));
+		return true;
+	}
+	return false;
+}
+
+} // namespace ASDTool
diff --git a/Engine/Source/Programs/Windows/ASDTool/Private/ASDToolShaderUtils.h b/Engine/Source/Programs/Windows/ASDTool/Private/ASDToolShaderUtils.h
new file mode 100644
--- /dev/null
+++ b/Engine/Source/Programs/Windows/ASDTool/Private/ASDToolShaderUtils.h
@@ -0,0 +1,107 @@
+// Copyright Epic Games, Inc. All Rights Reserved.
+
+#pragma once
+
+#include "CoreMinimal.h"
+#include "ShaderCodeArchive.h"
+
+/**
+ * ASDTool shader archive utilities for iterating cooked .ushaderbytecode archives.
+ *
+ * Provides:
+ * - FShaderKeyEntry:      Minimal per-shader info sorted by hash (built from .shk data externally)
+ * - FShaderArchiveEntry:  Per-shader data passed to the visitor during archive iteration
+ * - ProcessArchive:       Iterate a .ushaderbytecode file, visiting each unique matched shader
+ * - FindOptionalData:     Search a shader's optional data section for a key
+ * - ExtractResourceCounts: Extract FShaderCodePackedResourceCounts from optional data
+ */
+
+namespace ASDTool
+{
+
+//--------------------------------------------------------------------------------------------------
+// FShaderKeyEntry - Minimal info about a shader, sorted by hash for merge-join with archives
+//--------------------------------------------------------------------------------------------------
+struct FShaderKeyEntry
+{
+	FShaderHash Hash;
+};
+
+//--------------------------------------------------------------------------------------------------
+// FShaderArchiveEntry - Per-shader data extracted from a cooked archive, passed to the visitor
+//--------------------------------------------------------------------------------------------------
+struct FShaderArchiveEntry
+{
+	/** The UE shader output hash */
+	FShaderHash Hash;
+
+	/** Shader frequency (VS, PS, CS, etc.) */
+	EShaderFrequency Frequency = SF_NumFrequencies;
+
+	/** The raw shader code bytes (after SRT, before optional data).
+	 *  For D3D12 this is the full DXContainer (identified by 'DXBC' magic), which may
+	 *  contain a DXIL part (SM6+) or legacy DXBC bytecode (SM5). */
+	TArrayView64<const uint8> ShaderCode;
+
+	/** The optional data appended after shader code */
+	TArrayView64<const uint8> OptionalData;
+
+	/** Owning buffer for the decompressed data (caller must keep alive during callback) */
+	TArray<uint8>* UncompressedBuffer = nullptr;
+};
+
+/** Visitor callback type. Called once per unique matched shader. */
+using FShaderVisitor = TFunctionRef<void(FShaderArchiveEntry& Entry)>;
+
+/**
+ * Open a .ushaderbytecode archive and iterate over all shaders, matching them against
+ * the provided sorted SHK key entries. Calls the visitor for each unique, previously-unseen shader.
+ *
+ * @param ArchivePath       Path to the .ushaderbytecode file
+ * @param ShaderKeys        Sorted shader key entries (must be sorted by Hash)
+ * @param AlreadySeen       Set of hashes already processed (updated by this function)
+ * @param Visitor           Called for each unique shader entry
+ * @param bSingleThreaded   If true, process sequentially
+ * @param OutNumNotFound    Number of archive shaders not found in SHK data
+ * @param OutNumDuplicates  Number of shaders skipped as duplicates
+ * @return Number of unique shaders visited
+ */
+int32 ProcessArchive(
+	const FString& ArchivePath,
+	const TArray<FShaderKeyEntry>& ShaderKeys,
+	TSet<FShaderHash>& AlreadySeen,
+	FShaderVisitor Visitor,
+	bool bSingleThreaded = false,
+	int32* OutNumNotFound = nullptr,
+	int32* OutNumDuplicates = nullptr
+);
+
+/**
+ * Search the optional data section for a specific key.
+ * The OptionalData view must already be trimmed to exactly the optional data bytes,
+ * with no trailing size sentinel - as returned by ProcessArchive via DecompressAndExtract.
+ * Format: [Key(1 byte), Size(4 bytes), Data(Size bytes), ...] (repeating)
+ *
+ * @param OptionalData    View into the optional data bytes (trimmed, no trailing size uint32)
+ * @param InKey           The key to search for
+ * @param ValueSize       Expected size of the value
+ * @return Pointer to the value data, or nullptr if not found
+ */
+const uint8* FindOptionalData(
+	TArrayView64<const uint8> OptionalData,
+	EShaderOptionalDataKey InKey,
+	uint32 ValueSize
+);
+
+/**
+ * Extract FShaderCodePackedResourceCounts from shader optional data.
+ * @param OptionalData    View into the optional data bytes
+ * @param OutCounts       Output resource counts
+ * @return true if found and extracted
+ */
+bool ExtractResourceCounts(
+	TArrayView64<const uint8> OptionalData,
+	FShaderCodePackedResourceCounts& OutCounts
+);
+
+} // namespace ASDTool
diff --git a/Engine/Source/Programs/Windows/ASDTool/Private/ASDToolUtil.cpp b/Engine/Source/Programs/Windows/ASDTool/Private/ASDToolUtil.cpp
new file mode 100644
--- /dev/null
+++ b/Engine/Source/Programs/Windows/ASDTool/Private/ASDToolUtil.cpp
@@ -0,0 +1,325 @@
+// Copyright Epic Games, Inc. All Rights Reserved.
+
+#include "ASDToolUtil.h"
+#include "ASDToolCommands.h"
+
+#include "Misc/Paths.h"
+#include "Misc/FileHelper.h"
+#include "HAL/FileManager.h"
+#include "HAL/PlatformProcess.h"
+#include "HAL/PlatformTime.h"
+
+#include "Compression/OodleDataCompression.h"
+
+namespace ASDTool
+{
+
+//--------------------------------------------------------------------------------------------------
+// Process execution
+//--------------------------------------------------------------------------------------------------
+
+int32 RunProcess(const FString& ExePath, const FString& Args, FString& OutOutput, bool bLogOutput)
+{
+	if (!bLogOutput)
+	{
+		int32 ExitCode = -1;
+		FString StdOut;
+		FString StdErr;
+		FPlatformProcess::ExecProcess(*ExePath, *Args, &ExitCode, &StdOut, &StdErr);
+		OutOutput = StdOut;
+		if (!StdErr.IsEmpty())
+		{
+			if (!OutOutput.IsEmpty()) OutOutput += TEXT("\n");
+			OutOutput += StdErr;
+		}
+		return ExitCode;
+	}
+
+	// Streaming: use CreateProc with pipes so we can log output in real-time
+	int32 ExitCode = -1;
+	void* ReadPipe  = nullptr;
+	void* WritePipe = nullptr;
+
+	if (!FPlatformProcess::CreatePipe(ReadPipe, WritePipe))
+	{
+		UE_LOG(LogASDTool, Error, TEXT("RunProcess: Failed to create pipe for '%s'"), *ExePath);
+		return -1;
+	}
+
+	// NOTE: WritePipe is passed to CreateProc for the child's stdout and remains open in the parent
+	// until ClosePipe at function exit. Ideally the parent would close its write-end immediately
+	// after spawn to ensure proper EOF semantics, but FPlatformProcess does not expose a
+	// CloseWritePipe API -- ClosePipe closes both ends together. This is a UE API limitation.
+	// In practice the read loop terminates via IsProcRunning rather than EOF, so no hang occurs.
+	FProcHandle Proc = FPlatformProcess::CreateProc(
+		*ExePath, *Args,
+		false,    // bLaunchDetached
+		true,     // bLaunchHidden
+		true,     // bLaunchReallyHidden
+		nullptr,  // OutProcessID
+		0,        // PriorityModifier
+		nullptr,  // OptionalWorkingDirectory
+		WritePipe);
+
+	if (Proc.IsValid())
+	{
+		OutOutput.Empty();
+		FString LineBuffer;
+		FString LastProgressLine;
+
+		while (FPlatformProcess::IsProcRunning(Proc))
+		{
+			FString NewOutput = FPlatformProcess::ReadPipe(ReadPipe);
+			if (!NewOutput.IsEmpty())
+			{
+				OutOutput  += NewOutput;
+				LineBuffer += NewOutput;
+
+				// Process all complete lines (delimited by \n or \r)
+				// The compiler uses \r for progress updates on the same line
+				for (;;)
+				{
+					int32 NewlineIdx = INDEX_NONE;
+					int32 CRIdx      = INDEX_NONE;
+					LineBuffer.FindChar(TCHAR('\n'), NewlineIdx);
+					LineBuffer.FindChar(TCHAR('\r'), CRIdx);
+
+					int32 SplitIdx = INDEX_NONE;
+					if      (NewlineIdx != INDEX_NONE && CRIdx != INDEX_NONE) { SplitIdx = FMath::Min(NewlineIdx, CRIdx); }
+					else if (NewlineIdx != INDEX_NONE)                        { SplitIdx = NewlineIdx; }
+					else if (CRIdx      != INDEX_NONE)                        { SplitIdx = CRIdx; }
+					else                                                       { break; }
+
+					FString Line = LineBuffer.Left(SplitIdx).TrimEnd();
+					LineBuffer.RightChopInline(SplitIdx + 1);
+					if (Line.IsEmpty())
+					{
+						continue;
+					}
+
+					// Detect progress lines (e.g. "00:00:01: 23 / 308") -- overwrite in console
+					if (Line.Contains(TEXT(" / ")) && (Line[0] == TCHAR('0') || FChar::IsDigit(Line[0])))
+					{
+						FPlatformMisc::LocalPrint(*FString::Printf(TEXT("\r    %s    "), *Line));
+						LastProgressLine = Line;
+					}
+					else
+					{
+						if (!LastProgressLine.IsEmpty())
+						{
+							FPlatformMisc::LocalPrint(TEXT("\r"));
+							UE_LOG(LogASDTool, Display, TEXT("    %s"), *LastProgressLine);
+							LastProgressLine.Empty();
+						}
+						UE_LOG(LogASDTool, Display, TEXT("    %s"), *Line);
+					}
+				}
+			}
+			FPlatformProcess::Sleep(0.05f);
+		}
+
+		FString FinalOutput = FPlatformProcess::ReadPipe(ReadPipe);
+		if (!FinalOutput.IsEmpty())
+		{
+			OutOutput  += FinalOutput;
+			LineBuffer += FinalOutput;
+		}
+
+		if (!LastProgressLine.IsEmpty())
+		{
+			FPlatformMisc::LocalPrint(TEXT("\r"));
+			UE_LOG(LogASDTool, Display, TEXT("    %s"), *LastProgressLine);
+		}
+
+		FString RemainingLine = LineBuffer.TrimEnd();
+		if (!RemainingLine.IsEmpty())
+		{
+			UE_LOG(LogASDTool, Display, TEXT("    %s"), *RemainingLine);
+		}
+
+		FPlatformProcess::GetProcReturnCode(Proc, &ExitCode);
+		FPlatformProcess::CloseProc(Proc);
+	}
+
+	FPlatformProcess::ClosePipe(ReadPipe, WritePipe);
+	return ExitCode;
+}
+
+//--------------------------------------------------------------------------------------------------
+// AgilitySDK path helpers
+//--------------------------------------------------------------------------------------------------
+
+static FString GetAgilitySDKBinDir()
+{
+	return FPaths::ConvertRelativePathToFull(
+		FPaths::EngineDir() / TEXT("Source/ThirdParty/Windows/AgilitySDK") / TEXT(AGILITY_SDK_VERSION_STRING) / TEXT("Binaries/x64"));
+}
+
+FString FindAgilitySDKFile(const FString& Filename)
+{
+	FString Path = GetAgilitySDKBinDir() / Filename;
+	if (FPaths::FileExists(Path))
+	{
+		return Path;
+	}
+	UE_LOG(LogASDTool, Warning, TEXT("AgilitySDK file not found: %s"), *Path);
+	return FString();
+}
+
+FString FindDefaultCompilerPluginDir()
+{
+	FString Path = FPaths::ConvertRelativePathToFull(
+		FPaths::EngineDir() / TEXT("Binaries/ThirdParty/D3D12CompilerPlugins"));
+	if (FPaths::DirectoryExists(Path))
+	{
+		return Path;
+	}
+	return FString();
+}
+
+//--------------------------------------------------------------------------------------------------
+// Kraken compression
+//--------------------------------------------------------------------------------------------------
+
+bool CompressFile(const FString& FilePath, uint32 Magic)
+{
+	TArray64<uint8> UncompressedData;
+	if (!FFileHelper::LoadFileToArray(UncompressedData, *FilePath))
+	{
+		UE_LOG(LogASDTool, Error, TEXT("Failed to read file for compression: %s"), *FilePath);
+		return false;
+	}
+
+	int64 UncompressedSize     = UncompressedData.Num();
+	int64 CompressedBufferSize = FOodleDataCompression::CompressedBufferSizeNeeded(UncompressedSize);
+
+	TArray64<uint8> CompressedData;
+	CompressedData.SetNumUninitialized(CompressedBufferSize);
+
+	double CompressStart = FPlatformTime::Seconds();
+	int64 CompressedSize = FOodleDataCompression::CompressParallel(
+		CompressedData.GetData(), CompressedBufferSize,
+		UncompressedData.GetData(), UncompressedSize,
+		FOodleDataCompression::ECompressor::Kraken,
+		FOodleDataCompression::ECompressionLevel::Optimal2,
+		false,             // CompressIndependentChunks = false for best ratio
+		32 * 1024 * 1024   // MinChunkSize = 32 MB for better ratio vs parallelism tradeoff
+	);
+
+	if (CompressedSize <= 0)
+	{
+		UE_LOG(LogASDTool, Error, TEXT("Kraken compression failed for: %s"), *FilePath);
+		return false;
+	}
+
+	FString OutputPath = FilePath + TEXT(".oodle");
+
+	FCompressedHeader Header;
+	Header.Magic            = Magic;
+	Header.Version          = COMPRESSED_VERSION;
+	Header.UncompressedSize = UncompressedSize;
+	Header.CompressedSize   = CompressedSize;
+
+	TArray64<uint8> OutputData;
+	OutputData.SetNumUninitialized(sizeof(FCompressedHeader) + CompressedSize);
+	FMemory::Memcpy(OutputData.GetData(),                              &Header,                  sizeof(FCompressedHeader));
+	FMemory::Memcpy(OutputData.GetData() + sizeof(FCompressedHeader),  CompressedData.GetData(), CompressedSize);
+
+	if (!FFileHelper::SaveArrayToFile(OutputData, *OutputPath))
+	{
+		UE_LOG(LogASDTool, Error, TEXT("Failed to write compressed file: %s"), *OutputPath);
+		return false;
+	}
+
+	double CompressSeconds = FPlatformTime::Seconds() - CompressStart;
+	double Ratio           = (double)UncompressedSize / (double)CompressedSize;
+	double EncodeMBps      = (UncompressedSize / (1024.0 * 1024.0)) / FMath::Max(CompressSeconds, 0.001);
+	UE_LOG(LogASDTool, Display, TEXT("  Compressed: %s (%.1f MB -> %.1f MB, %.1f:1, %.1fs, %.1f MB/s)"),
+		*FPaths::GetCleanFilename(OutputPath),
+		UncompressedSize / (1024.0 * 1024.0),
+		CompressedSize   / (1024.0 * 1024.0),
+		Ratio, CompressSeconds, EncodeMBps);
+
+	// Delete the original -- only keep the compressed version
+	IFileManager::Get().Delete(*FilePath);
+	return true;
+}
+
+bool DecompressFile(const FString& CompressedPath, const FString& OutputPath, uint32 Magic)
+{
+	TArray64<uint8> FileData;
+	if (!FFileHelper::LoadFileToArray(FileData, *CompressedPath))
+	{
+		UE_LOG(LogASDTool, Error, TEXT("Failed to read compressed file: %s"), *CompressedPath);
+		return false;
+	}
+
+	if (FileData.Num() < (int64)sizeof(FCompressedHeader))
+	{
+		UE_LOG(LogASDTool, Error, TEXT("Compressed file too small: %s"), *CompressedPath);
+		return false;
+	}
+
+	const FCompressedHeader* Header = reinterpret_cast<const FCompressedHeader*>(FileData.GetData());
+
+	if (Header->Magic != Magic)
+	{
+		UE_LOG(LogASDTool, Error, TEXT("Invalid compressed file magic in: %s"), *CompressedPath);
+		return false;
+	}
+
+	if (Header->Version != COMPRESSED_VERSION)
+	{
+		UE_LOG(LogASDTool, Error, TEXT("Unsupported compressed file version %u (expected %u) in: %s"),
+			Header->Version, COMPRESSED_VERSION, *CompressedPath);
+		return false;
+	}
+
+	const uint8* CompressedData = FileData.GetData() + sizeof(FCompressedHeader);
+	int64 CompressedSize        = Header->CompressedSize;
+	int64 UncompressedSize      = Header->UncompressedSize;
+
+	// Validate compressed data fits within the file.
+	// Written as subtraction to avoid any possibility of signed overflow in the addition.
+	if (CompressedSize <= 0 || CompressedSize > FileData.Num() - int64(sizeof(FCompressedHeader)))
+	{
+		UE_LOG(LogASDTool, Error, TEXT("Invalid compressed size in: %s"), *CompressedPath);
+		return false;
+	}
+
+	// Validate uncompressed size is reasonable (max 4 GB)
+	static constexpr int64 MAX_UNCOMPRESSED_SIZE = 4LL * 1024 * 1024 * 1024;
+	if (UncompressedSize <= 0 || UncompressedSize > MAX_UNCOMPRESSED_SIZE)
+	{
+		UE_LOG(LogASDTool, Error, TEXT("Invalid uncompressed size (%lld) in: %s"), UncompressedSize, *CompressedPath);
+		return false;
+	}
+
+	TArray64<uint8> UncompressedData;
+	UncompressedData.SetNumUninitialized(UncompressedSize);
+
+	if (!FOodleDataCompression::Decompress(
+		UncompressedData.GetData(), UncompressedSize,
+		CompressedData, CompressedSize))
+	{
+		UE_LOG(LogASDTool, Error, TEXT("Kraken decompression failed for: %s"), *CompressedPath);
+		return false;
+	}
+
+	IFileManager::Get().MakeDirectory(*FPaths::GetPath(OutputPath), true);
+
+	if (!FFileHelper::SaveArrayToFile(UncompressedData, *OutputPath))
+	{
+		UE_LOG(LogASDTool, Error, TEXT("Failed to write decompressed file: %s"), *OutputPath);
+		return false;
+	}
+
+	UE_LOG(LogASDTool, Display, TEXT("  Decompressed: %s (%.1f MB -> %.1f MB)"),
+		*FPaths::GetCleanFilename(OutputPath),
+		CompressedSize   / (1024.0 * 1024.0),
+		UncompressedSize / (1024.0 * 1024.0));
+
+	return true;
+}
+
+} // namespace ASDTool
diff --git a/Engine/Source/Programs/Windows/ASDTool/Private/ASDToolUtil.h b/Engine/Source/Programs/Windows/ASDTool/Private/ASDToolUtil.h
new file mode 100644
--- /dev/null
+++ b/Engine/Source/Programs/Windows/ASDTool/Private/ASDToolUtil.h
@@ -0,0 +1,70 @@
+// Copyright Epic Games, Inc. All Rights Reserved.
+
+#pragma once
+
+#include "CoreMinimal.h"
+
+namespace ASDTool
+{
+
+//--------------------------------------------------------------------------------------------------
+// Process execution
+//--------------------------------------------------------------------------------------------------
+
+/** Run a process and capture stdout. Returns exit code.
+ *  If bLogOutput is true, logs lines in real-time with progress-line (\r) handling. */
+int32 RunProcess(const FString& ExePath, const FString& Args, FString& OutOutput, bool bLogOutput = false);
+
+//--------------------------------------------------------------------------------------------------
+// AgilitySDK path helpers
+//
+// AGILITY_SDK_VERSION_STRING is injected by ASDTool.Build.cs from AgilitySDK.DefaultVersion,
+// keeping it in sync with the rest of the engine without hardcoding here.
+//
+// For normal D3D12 device creation, D3D12Core.dll does NOT need a path helper -- the OS loader
+// finds it automatically via the D3D12SDKPath export declared in UnrealAgilitySDKLink.inl.
+// However, FSODBFactory::Initialize() loads D3D12Core.dll directly (bypassing D3D12.dll) to
+// access ID3D12StateObjectDatabaseFactory, which is only exposed in the Agility SDK build.
+// FindAgilitySDKFile("D3D12Core.dll") is used for that case, keeping all Agility SDK binary
+// resolution on a single code path.
+//--------------------------------------------------------------------------------------------------
+
+/** Returns the absolute path to a named file in the AgilitySDK x64 binaries directory,
+ *  or empty string if not found. Use for D3D12StateObjectCompiler.exe and .dll. */
+FString FindAgilitySDKFile(const FString& Filename);
+
+/** Returns the absolute path to the default IHV compiler plugin directory, or empty string if not found. */
+FString FindDefaultCompilerPluginDir();
+
+//--------------------------------------------------------------------------------------------------
+// Kraken compression
+//
+// Files are compressed with Kraken Optimal2 and prefixed with FCompressedHeader.
+// The compressed output path is always InputPath + ".oodle".
+//--------------------------------------------------------------------------------------------------
+
+/** Header written at the start of every compressed file. */
+struct FCompressedHeader
+{
+	uint32 Magic;            // See PSDB_COMPRESSED_MAGIC / SODB_COMPRESSED_MAGIC
+	uint32 Version;          // See COMPRESSED_VERSION
+	int64  UncompressedSize;
+	int64  CompressedSize;
+};
+
+static constexpr uint32 COMPRESSED_VERSION     = 1;
+// Magic constants use explicit byte shifts rather than multi-character literals to guarantee
+// stable on-disk values regardless of compiler or endianness (multi-char literals are
+// implementation-defined in C++ and should not be used for persisted file format identifiers).
+static constexpr uint32 PSDB_COMPRESSED_MAGIC  = (uint32('C') << 24) | (uint32('P') << 16) | (uint32('D') << 8) | uint32('B');
+static constexpr uint32 SODB_COMPRESSED_MAGIC  = (uint32('C') << 24) | (uint32('S') << 16) | (uint32('D') << 8) | uint32('B');  // Reserved for future SODB compression
+
+/** Compress a file with Kraken Optimal2, writing to FilePath + ".oodle".
+ *  Deletes the original file on success. Returns true on success. */
+bool CompressFile(const FString& FilePath, uint32 Magic);
+
+/** Decompress a .oodle file into OutputPath.
+ *  Creates the output directory if needed. Returns true on success. */
+bool DecompressFile(const FString& CompressedPath, const FString& OutputPath, uint32 Magic);
+
+} // namespace ASDTool
diff --git a/Engine/Source/Programs/Windows/ASDTool/Private/CompilePSDBs.cpp b/Engine/Source/Programs/Windows/ASDTool/Private/CompilePSDBs.cpp
new file mode 100644
--- /dev/null
+++ b/Engine/Source/Programs/Windows/ASDTool/Private/CompilePSDBs.cpp
@@ -0,0 +1,1202 @@
+// Copyright Epic Games, Inc. All Rights Reserved.
+
+#include "ASDToolCommands.h"
+
+#include "Misc/CommandLine.h"
+#include "Misc/Parse.h"
+#include "Misc/Paths.h"
+#include "HAL/FileManager.h"
+
+
+// initguid.h must be included before any D3D12 headers to instantiate GUIDs
+#include "Windows/AllowWindowsPlatformTypes.h"
+THIRD_PARTY_INCLUDES_START
+#include <initguid.h>
+#include <d3d12.h>
+#include <d3d12compiler.h>
+THIRD_PARTY_INCLUDES_END
+#include "Windows/HideWindowsPlatformTypes.h"
+
+#include "SODBFile.h"
+#include "ASDToolUtil.h"
+
+#include "Async/Async.h"
+#include "Async/ParallelFor.h"
+
+
+namespace ASDTool
+{
+
+//--------------------------------------------------------------------------------------------------
+// Data types
+//--------------------------------------------------------------------------------------------------
+
+struct FAdapterFamily
+{
+	FString Name;              // e.g. "AMD_NAVI31:gfx1100"
+	FString CompilerVersion;   // e.g. "32.0.23027.2005"
+	TArray<uint64> ABIVersions;
+};
+
+struct FCompilerPlugin
+{
+	FString DllPath;
+	FString IHVName;           // e.g. "AMD"
+	FString CompilerVersion;   // folder name, e.g. "32.0.23027.2005"
+	TArray<FAdapterFamily> Families;
+};
+
+struct FLocalAdapterInfo
+{
+	FString FamilyName;
+	uint64 MinABIVersion = 0;
+	uint64 MaxABIVersion = 0;
+};
+
+struct FCompileWorkItem
+{
+	// Raw pointers into OutPlugins (TArray<FCompilerPlugin>). Safe as long as OutPlugins is not
+	// mutated after BuildCompileWorkItems is called. Will be replaced with indices or handles
+	// when this is refactored for Horde job distribution.
+	const FCompilerPlugin* Plugin;
+	const FAdapterFamily*  Family;
+	uint64                 ABIVersion;
+	FString                SODBPath;
+	FString                SODBName;
+	FString                OutputPath;
+};
+
+//--------------------------------------------------------------------------------------------------
+// Command-line arguments
+//--------------------------------------------------------------------------------------------------
+
+struct FCompilePSDBsArgs
+{
+	FString SODBPath;
+	FString OutputDir;
+	FString PluginDir;
+	FString CompilerExe;
+	FString AppName;
+	FString ExeFilename;
+	FString EngineName;
+	bool    bListOnly        = false;
+	bool    bClean           = false;
+	bool    bLocalAdapter    = false;
+	bool    bCompress        = false;
+	bool    bUseCompilerExe  = false;
+	int32   NumThreads       = 0;  // 0 = auto (all hardware cores)
+	FString AdapterFamily;   // -AdapterFamily=: bypass adapter filter, compile only this family
+	uint64  ABI        = 0;  // -ABI=: ABI version to use with -AdapterFamily (0 = use DLL-reported ABI)
+};
+
+//--------------------------------------------------------------------------------------------------
+// FCompilePSDBsStats
+// Accumulated across CompileWithExe/CompileWithCOM and CompressAllPSDBs.
+// ElapsedSeconds covers total wall time, set in CompilePSDBs().
+//--------------------------------------------------------------------------------------------------
+
+struct FCompilePSDBsStats
+{
+	int32  PSDBsCompiled   = 0;
+	int32  PSDBsFailed     = 0;
+	int32  PSDBsCompressed = 0;
+	int32  CompressFailed  = 0;
+	double ElapsedSeconds  = 0.0;
+};
+
+//--------------------------------------------------------------------------------------------------
+// Plugin discovery & enumeration
+//--------------------------------------------------------------------------------------------------
+
+/** Scan PluginDir for compiler plugin DLLs.
+ *  Expected structure: <PluginDir>/<IHV>/<CompilerVersion>/<plugin>.dll */
+static void DiscoverPlugins(const FString& PluginDir, TArray<FCompilerPlugin>& OutPlugins)
+{
+	TArray<FString> DllFiles;
+	IFileManager::Get().FindFilesRecursive(DllFiles, *PluginDir, TEXT("*.dll"), true, false);
+
+	for (const FString& DllPath : DllFiles)
+	{
+		FString RelPath = DllPath;
+		FPaths::MakePathRelativeTo(RelPath, *(PluginDir / TEXT("")));
+
+		TArray<FString> Parts;
+		RelPath.ParseIntoArray(Parts, TEXT("/"));
+		if (Parts.Num() < 3)
+		{
+			RelPath.ParseIntoArray(Parts, TEXT("\\"));
+		}
+
+		if (Parts.Num() >= 3)
+		{
+			FCompilerPlugin Plugin;
+			Plugin.DllPath         = FPaths::ConvertRelativePathToFull(DllPath);
+			Plugin.IHVName         = Parts[0];
+			Plugin.CompilerVersion = Parts[1];
+			OutPlugins.Add(MoveTemp(Plugin));
+		}
+		else
+		{
+			UE_LOG(LogASDTool, Warning, TEXT("Skipping DLL with unexpected path structure: %s"), *DllPath);
+		}
+	}
+}
+
+/** Parse the output of 'D3D12StateObjectCompiler.exe list --plugin <dll>'. */
+static bool ParseListOutput(const FString& Output, TArray<FAdapterFamily>& OutFamilies)
+{
+	TArray<FString> Lines;
+	Output.ParseIntoArrayLines(Lines);
+
+	FAdapterFamily* CurrentFamily = nullptr;
+
+	for (const FString& Line : Lines)
+	{
+		FString Trimmed = Line.TrimStartAndEnd();
+
+		if (Trimmed.Len() > 0 && FChar::IsDigit(Trimmed[0]) && Trimmed.EndsWith(TEXT(":")))
+		{
+			CurrentFamily = &OutFamilies.Emplace_GetRef();
+			continue;
+		}
+		if (!CurrentFamily)
+		{
+			continue;
+		}
+
+		if      (Trimmed.StartsWith(TEXT("Adapter Family:")))
+		{
+			CurrentFamily->Name            = Trimmed.Mid(16).TrimStartAndEnd();
+		}
+		else if (Trimmed.StartsWith(TEXT("Compiler Version:")))
+		{
+			CurrentFamily->CompilerVersion  = Trimmed.Mid(18).TrimStartAndEnd();
+		}
+		else if (Trimmed.StartsWith(TEXT("ABI Versions:")))
+		{
+			FString ABIStr = Trimmed.Mid(14).TrimStartAndEnd();
+			ABIStr.RemoveFromStart(TEXT("{"));
+			ABIStr.RemoveFromEnd(TEXT("}"));
+			TArray<FString> ABIParts;
+			ABIStr.ParseIntoArray(ABIParts, TEXT(","));
+			for (const FString& Part : ABIParts)
+			{
+				uint64 ABI = FCString::Strtoui64(*Part.TrimStartAndEnd(), nullptr, 10);
+				if (ABI > 0)
+				{
+					CurrentFamily->ABIVersions.Add(ABI);
+				}
+			}
+		}
+	}
+	return OutFamilies.Num() > 0;
+}
+
+/** Enumerate adapter families for a plugin by running the compiler exe. */
+static bool EnumeratePlugin(const FString& CompilerExe, FCompilerPlugin& Plugin)
+{
+	FString Output;
+	int32 ExitCode = RunProcess(CompilerExe, FString::Printf(TEXT("list --plugin \"%s\""), *Plugin.DllPath), Output);
+	if (ExitCode != 0)
+	{
+		UE_LOG(LogASDTool, Warning, TEXT("Failed to enumerate plugin '%s' (exit code %d)"), *Plugin.DllPath, ExitCode);
+		return false;
+	}
+	return ParseListOutput(Output, Plugin.Families);
+}
+
+/** Query installed adapters via 'list --adapters'. Returns family names with ABI ranges. */
+static bool EnumerateLocalAdapters(const FString& CompilerExe, TArray<FLocalAdapterInfo>& OutAdapters)
+{
+	FString Output;
+	int32 ExitCode = RunProcess(CompilerExe, TEXT("list --adapters"), Output);
+	if (ExitCode != 0)
+	{
+		UE_LOG(LogASDTool, Error, TEXT("Failed to enumerate local adapters (exit code %d)"), ExitCode);
+		return false;
+	}
+
+	TArray<FString> Lines;
+	Output.ParseIntoArrayLines(Lines);
+
+	FLocalAdapterInfo CurrentAdapter;
+	for (const FString& Line : Lines)
+	{
+		FString Trimmed = Line.TrimStartAndEnd();
+		if      (Trimmed.StartsWith(TEXT("Adapter Family:")))
+		{
+			CurrentAdapter.FamilyName  = Trimmed.Mid(16).TrimStartAndEnd();
+		}
+		else if (Trimmed.StartsWith(TEXT("MinimumABISupportVersion:")))
+		{
+			CurrentAdapter.MinABIVersion = FCString::Strtoui64(*Trimmed.Mid(25).TrimStartAndEnd(), nullptr, 10);
+		}
+		else if (Trimmed.StartsWith(TEXT("MaximumABISupportVersion:")))
+		{
+			CurrentAdapter.MaxABIVersion = FCString::Strtoui64(*Trimmed.Mid(25).TrimStartAndEnd(), nullptr, 10);
+			if (!CurrentAdapter.FamilyName.IsEmpty() && CurrentAdapter.MaxABIVersion > 0)
+			{
+				bool bAlreadyAdded = OutAdapters.ContainsByPredicate([&](const FLocalAdapterInfo& E) { return E.FamilyName == CurrentAdapter.FamilyName; });
+				if (!bAlreadyAdded)
+				{
+					OutAdapters.Add(CurrentAdapter);
+				}
+			}
+			CurrentAdapter = FLocalAdapterInfo();
+		}
+	}
+	return OutAdapters.Num() > 0;
+}
+
+static int32 DiscoverAndEnumeratePlugins(
+	const FCompilePSDBsArgs& Args,
+	TArray<FCompilerPlugin>& OutPlugins,
+	int32& OutTotalFamilies)
+{
+	DiscoverPlugins(Args.PluginDir, OutPlugins);
+	if (OutPlugins.IsEmpty())
+	{
+		UE_LOG(LogASDTool, Error, TEXT("No compiler plugin DLLs found in '%s'"), *Args.PluginDir);
+		return 1;
+	}
+
+	UE_LOG(LogASDTool, Display, TEXT("Found %d plugin DLL(s). Enumerating adapter families..."), OutPlugins.Num());
+
+	OutTotalFamilies = 0;
+	for (FCompilerPlugin& Plugin : OutPlugins)
+	{
+		if (EnumeratePlugin(Args.CompilerExe, Plugin))
+		{
+			UE_LOG(LogASDTool, Display, TEXT("  %s/%s: %d adapter families"),
+				   *Plugin.IHVName, *Plugin.CompilerVersion, Plugin.Families.Num());
+			for (const FAdapterFamily& Family : Plugin.Families)
+			{
+				FString ABIStr;
+				for (uint64 ABI : Family.ABIVersions)
+				{
+					ABIStr += (ABIStr.IsEmpty() ? TEXT("") : TEXT(", ")) + FString::Printf(TEXT("%llu"), ABI);
+				}
+				UE_LOG(LogASDTool, Display, TEXT("    %-40s  ABI: {%s}  Compiler: %s"), *Family.Name, *ABIStr, *Family.CompilerVersion);
+			}
+			OutTotalFamilies += Plugin.Families.Num();
+		}
+		else
+		{
+			UE_LOG(LogASDTool, Warning, TEXT("  %s/%s: enumeration failed, skipping"), *Plugin.IHVName, *Plugin.CompilerVersion);
+		}
+	}
+
+	UE_LOG(LogASDTool, Display, TEXT("Total: %d plugins, %d adapter families"), OutPlugins.Num(), OutTotalFamilies);
+	return 0;
+}
+
+static int32 FilterLocalAdapters(const FCompilePSDBsArgs& Args, TArray<FLocalAdapterInfo>& OutAdapters)
+{
+	if (!EnumerateLocalAdapters(Args.CompilerExe, OutAdapters))
+	{
+		UE_LOG(LogASDTool, Error, TEXT("Failed to enumerate local adapters for -LocalAdapter filter"));
+		return 1;
+	}
+
+	UE_LOG(LogASDTool, Display, TEXT("Local adapter filter -- compiling only for:"));
+	for (const FLocalAdapterInfo& Adapter : OutAdapters)
+	{
+		UE_LOG(LogASDTool, Display, TEXT("  %s (ABI range: %llu - %llu)"),
+			   *Adapter.FamilyName, Adapter.MinABIVersion, Adapter.MaxABIVersion);
+	}
+	return 0;
+}
+
+
+//--------------------------------------------------------------------------------------------------
+// Exe compilation
+//--------------------------------------------------------------------------------------------------
+
+/** Compile a single work item via D3D12StateObjectCompiler.exe. */
+static bool CompileSODB(const FString& CompilerExe, const FCompileWorkItem& Item,
+						const FString& AppName, const FString& ExeFilename, const FString& EngineName)
+{
+	IFileManager::Get().MakeDirectory(*FPaths::GetPath(Item.OutputPath), true);
+
+	FString Args = FString::Printf(
+		TEXT("compile \"%s\" \"%s\" --plugin \"%s\" --adapter-family \"%s\" --abi %llu --name \"%s\" --exe-filename \"%s\" --engine \"%s\""),
+		*Item.SODBPath, *Item.OutputPath, *Item.Plugin->DllPath,
+		*Item.Family->Name, Item.ABIVersion, *AppName, *ExeFilename, *EngineName);
+
+	FString Output;
+	int32 ExitCode = RunProcess(CompilerExe, Args, Output, true);
+	if (ExitCode != 0)
+	{
+		UE_LOG(LogASDTool, Error, TEXT("  FAILED [%s, ABI %llu, %s] exit code %d"),
+			   *Item.Family->Name, Item.ABIVersion, *Item.SODBName, ExitCode);
+		return false;
+	}
+	return true;
+}
+
+static void CompileWithExe(
+	const FCompilePSDBsArgs& Args,
+	const TArray<FCompileWorkItem>& CompileItems,
+	FCompilePSDBsStats& OutStats)
+{
+	UE_LOG(LogASDTool, Display, TEXT("=== Compiling PSDBs (EXE) ==="));
+
+	for (int32 i = 0; i < CompileItems.Num(); ++i)
+	{
+		const FCompileWorkItem& Item = CompileItems[i];
+		UE_LOG(LogASDTool, Display, TEXT("[%d/%d] %s | ABI %llu | %s"),
+			   i + 1, CompileItems.Num(), *Item.Family->Name, Item.ABIVersion, *Item.SODBName);
+
+		if (CompileSODB(Args.CompilerExe, Item, Args.AppName, Args.ExeFilename, Args.EngineName))
+		{
+			++OutStats.PSDBsCompiled;
+		}
+		else
+		{
+			++OutStats.PSDBsFailed;
+		}
+	}
+}
+
+
+//--------------------------------------------------------------------------------------------------
+// COM-object compilation
+//--------------------------------------------------------------------------------------------------
+
+/** Per-thread compiler context for COM path. Each thread gets its own DLL copy (separate heap)
+ *  to eliminate cross-thread heap lock contention. */
+struct FThreadCompilerContext
+{
+	FString TempDir;
+	FString DllCopyPath;
+
+	~FThreadCompilerContext()
+	{
+		for (auto& Pair : CachedFactories)
+		{
+			Pair.Value->Release();
+		}
+		CachedFactories.Empty();
+		if (CompilerModule)
+		{
+			FreeLibrary(CompilerModule);
+			CompilerModule = nullptr;
+		}
+		if (!DllCopyPath.IsEmpty() && FPaths::FileExists(DllCopyPath))
+		{
+			IFileManager::Get().Delete(*DllCopyPath);
+		}
+	}
+
+	bool Initialize(const FString& OriginalDllPath, const FString& InTempDir)
+	{
+		TempDir = InTempDir;
+		IFileManager::Get().MakeDirectory(*TempDir, true);
+
+		DllCopyPath = TempDir / TEXT("D3D12StateObjectCompiler.dll");
+		if (IFileManager::Get().Copy(*DllCopyPath, *OriginalDllPath) != COPY_OK)
+		{
+			UE_LOG(LogASDTool, Error, TEXT("Failed to copy D3D12StateObjectCompiler.dll to: %s"), *DllCopyPath);
+			return false;
+		}
+
+		CompilerModule = LoadLibrary(*DllCopyPath);
+		if (!CompilerModule)
+		{
+			UE_LOG(LogASDTool, Error, TEXT("Failed to load DLL copy from: %s"), *DllCopyPath);
+			return false;
+		}
+
+		CreateFactoryFunc = reinterpret_cast<PFN_D3D12_COMPILER_CREATE_FACTORY>(
+			reinterpret_cast<void*>(GetProcAddress(CompilerModule, "D3D12CompilerCreateFactory")));
+		if (!CreateFactoryFunc)
+		{
+			UE_LOG(LogASDTool, Error, TEXT("D3D12CompilerCreateFactory not found in DLL: %s"), *DllCopyPath);
+			return false;
+		}
+		return true;
+	}
+
+	ID3D12CompilerFactory* GetOrCreateFactory(const FString& PluginDllPath)
+	{
+		if (ID3D12CompilerFactory** Found = CachedFactories.Find(PluginDllPath))
+		{
+			return *Found;
+		}
+
+		ID3D12CompilerFactory* Factory = nullptr;
+		HRESULT hr = CreateFactoryFunc(*PluginDllPath, IID_PPV_ARGS(&Factory));
+		if (FAILED(hr) || !Factory)
+		{
+			UE_LOG(LogASDTool, Error, TEXT("D3D12CompilerCreateFactory failed for '%s': 0x%08X"), *PluginDllPath, hr);
+			return nullptr;
+		}
+		CachedFactories.Add(PluginDllPath, Factory);
+		return Factory;
+	}
+
+	ID3D12Compiler* CreateCompiler(const FString& PluginDllPath, const FString& FamilyName,
+		uint64 ABIVersion, const FString& DestPSDBPath, const D3D12_APPLICATION_DESC& AppDesc,
+		ID3D12CompilerCacheSession*& OutSession)
+	{
+		OutSession = nullptr;
+
+		ID3D12CompilerFactory* Factory = GetOrCreateFactory(PluginDllPath);
+		if (!Factory)
+		{
+			return nullptr;
+		}
+
+		UINT AdapterFamilyIndex = UINT_MAX;
+		for (UINT i = 0; ; ++i)
+		{
+			D3D12_ADAPTER_FAMILY AdapterFamily = {};
+			if (FAILED(Factory->EnumerateAdapterFamilies(i, &AdapterFamily)))
+			{
+				break;
+			}
+			if (FString(AdapterFamily.szAdapterFamily) == FamilyName)
+			{
+				AdapterFamilyIndex = i;
+				break;
+			}
+		}
+
+		if (AdapterFamilyIndex == UINT_MAX)
+		{
+			UE_LOG(LogASDTool, Error, TEXT("Adapter family '%s' not found in plugin '%s'"), *FamilyName, *PluginDllPath);
+			return nullptr;
+		}
+
+		D3D12_COMPILER_DATABASE_PATH DbPath = {};
+		DbPath.Types = D3D12_COMPILER_VALUE_TYPE_FLAGS_OBJECT_CODE | D3D12_COMPILER_VALUE_TYPE_FLAGS_METADATA;
+		DbPath.pPath = *DestPSDBPath;
+
+		D3D12_COMPILER_TARGET Target = {};
+		Target.AdapterFamilyIndex = AdapterFamilyIndex;
+		Target.ABIVersion         = ABIVersion;
+
+		HRESULT hr = Factory->CreateCompilerCacheSession(&DbPath, 1, &Target, &AppDesc, IID_PPV_ARGS(&OutSession));
+		if (FAILED(hr) || !OutSession)
+		{
+			UE_LOG(LogASDTool, Error, TEXT("CreateCompilerCacheSession failed: 0x%08X"), hr);
+			return nullptr;
+		}
+
+		ID3D12Compiler* Compiler = nullptr;
+		hr = Factory->CreateCompiler(OutSession, IID_PPV_ARGS(&Compiler));
+		if (FAILED(hr) || !Compiler)
+		{
+			UE_LOG(LogASDTool, Error, TEXT("CreateCompiler failed: 0x%08X"), hr);
+			OutSession->Release(); OutSession = nullptr;
+			return nullptr;
+		}
+		return Compiler;
+	}
+
+private:
+	HMODULE CompilerModule = nullptr;
+	PFN_D3D12_COMPILER_CREATE_FACTORY CreateFactoryFunc = nullptr;
+	TMap<FString, ID3D12CompilerFactory*> CachedFactories;
+};
+
+/** Per-thread compile loop. Pulls PSO indices from shared atomic queue until exhausted. */
+static void CompileOnThread(
+	FThreadCompilerContext& Ctx,
+	const FCompilePSDBsArgs& Args,
+	const FCompilerPlugin& Plugin,
+	const FAdapterFamily& Family,
+	uint64 ABIVersion,
+	const TArray<FSODBEntry>& SODBEntries,
+	const FString& DestPSDBPath,
+	std::atomic<int32>& NextIndex,
+	std::atomic<int32>& CompiledCount,
+	std::atomic<int32>& FailedCount,
+	std::atomic<int64>& TotalCompileCycles,
+	FCriticalSection& CreateCompilerLock)
+{
+	D3D12_APPLICATION_DESC AppDesc = {};
+	AppDesc.pExeFilename = *Args.ExeFilename;
+	AppDesc.pName        = *Args.AppName;
+	AppDesc.pEngineName  = *Args.EngineName;
+
+	ID3D12CompilerCacheSession* Session  = nullptr;
+	ID3D12Compiler*             Compiler = nullptr;
+	{
+		// Serialize CreateCompilerCacheSession -- amdxc64.dll is not thread-safe during init
+		FScopeLock Lock(&CreateCompilerLock);
+		Compiler = Ctx.CreateCompiler(Plugin.DllPath, Family.Name, ABIVersion, DestPSDBPath, AppDesc, Session);
+	}
+
+	if (!Compiler)
+	{
+		int32 Total = SODBEntries.Num();
+		while (NextIndex.fetch_add(1) < Total)
+		{
+			++FailedCount;
+		}
+		return;
+	}
+
+	int32 Index;
+	while ((Index = NextIndex.fetch_add(1)) < SODBEntries.Num())
+	{
+		const FSODBEntry& Entry = SODBEntries[Index];
+
+		D3D12_COMPILER_CACHE_GROUP_KEY GroupKey = {};
+		GroupKey.pKey    = Entry.Key.GetData();
+		GroupKey.KeySize = Entry.Key.Num();
+
+		// Build pipeline state stream with live pointers on the stack
+		D3D12_PIPELINE_STATE_STREAM_DESC StreamDesc = {};
+		struct { D3D12_PIPELINE_STATE_SUBOBJECT_TYPE Type; D3D12_SERIALIZED_ROOT_SIGNATURE_DESC Desc; } RootSigSubobject = {};
+		struct { D3D12_PIPELINE_STATE_SUBOBJECT_TYPE Type; D3D12_SHADER_BYTECODE Desc; }                CSSubobject       = {};
+		uint8 StreamBuffer[sizeof(RootSigSubobject) + sizeof(CSSubobject)];
+
+		// TODO: Graphics PSOs written via -DirectWrite have StreamData empty and StageDXBC populated.
+		// Building a graphics pipeline stream from per-stage DXBC (VS/PS/GS/MS/AS + root sig) is not
+		// yet implemented here. This will be addressed when graphics PSO compilation is productionized.
+		// For now, graphics entries from the direct SQLite path will be skipped (logged below).
+		if (Entry.bIsGraphics && Entry.StreamData.IsEmpty())
+		{
+			UE_LOG(LogASDTool, Verbose, TEXT("  Skipping graphics entry -- direct-write compile path not yet implemented (key size: %d)"), Entry.Key.Num());
+			++FailedCount;
+			continue;
+		}
+
+		if (Entry.StreamData.Num() > 0)
+		{
+			StreamDesc = Entry.GetDesc();
+		}
+		else if (Entry.StageDXBC[SF_Compute].Num() > 0)
+		{
+			uint8* WritePtr = StreamBuffer;
+
+			if (Entry.RootSigData.Num() > 0)
+			{
+				RootSigSubobject.Type                           = D3D12_PIPELINE_STATE_SUBOBJECT_TYPE_SERIALIZED_ROOT_SIGNATURE;
+				RootSigSubobject.Desc.pSerializedBlob           = Entry.RootSigData.GetData();
+				RootSigSubobject.Desc.SerializedBlobSizeInBytes = Entry.RootSigData.Num();
+				FMemory::Memcpy(WritePtr, &RootSigSubobject, sizeof(RootSigSubobject));
+				WritePtr += sizeof(RootSigSubobject);
+			}
+			else if (Entry.RootSigBlob)
+			{
+				RootSigSubobject.Type                           = D3D12_PIPELINE_STATE_SUBOBJECT_TYPE_SERIALIZED_ROOT_SIGNATURE;
+				RootSigSubobject.Desc.pSerializedBlob           = Entry.RootSigBlob->GetBufferPointer();
+				RootSigSubobject.Desc.SerializedBlobSizeInBytes = Entry.RootSigBlob->GetBufferSize();
+				FMemory::Memcpy(WritePtr, &RootSigSubobject, sizeof(RootSigSubobject));
+				WritePtr += sizeof(RootSigSubobject);
+			}
+
+			CSSubobject.Type                  = D3D12_PIPELINE_STATE_SUBOBJECT_TYPE_CS;
+			CSSubobject.Desc.pShaderBytecode  = Entry.StageDXBC[SF_Compute].GetData();
+			CSSubobject.Desc.BytecodeLength   = Entry.StageDXBC[SF_Compute].Num();
+			FMemory::Memcpy(WritePtr, &CSSubobject, sizeof(CSSubobject));
+			WritePtr += sizeof(CSSubobject);
+
+			StreamDesc.SizeInBytes                 = WritePtr - StreamBuffer;
+			StreamDesc.pPipelineStateSubobjectStream = StreamBuffer;
+		}
+		else
+		{
+			++FailedCount;
+			continue;
+		}
+
+		int64 CompileStart  = FPlatformTime::Cycles64();
+		HRESULT CompileHr   = Compiler->CompilePipelineState(&GroupKey, Entry.Version, &StreamDesc);
+		int64 CompileCycles = FPlatformTime::Cycles64() - CompileStart;
+
+		if (SUCCEEDED(CompileHr))
+		{
+			TotalCompileCycles.fetch_add(CompileCycles);
+			int32 Count = ++CompiledCount;
+			if (Count % 500 == 0)
+			{
+				UE_LOG(LogASDTool, Display, TEXT("    Progress: %d / %d compiled"), Count, SODBEntries.Num());
+			}
+		}
+		else
+		{
+			++FailedCount;
+		}
+	}
+
+	Compiler->Release();
+	Session->Release();
+}
+
+/** Compile all entries in SODBEntries into per-thread PSDBs using N independent DLL instances. */
+static bool CompileDirectCOM(
+	const FCompilePSDBsArgs& Args,
+	const FCompilerPlugin& Plugin,
+	const FAdapterFamily& Family,
+	uint64 ABIVersion,
+	const TArray<FSODBEntry>& SODBEntries,
+	const FString& DestPSDBPath)
+{
+	IFileManager::Get().MakeDirectory(*FPaths::GetPath(DestPSDBPath), true);
+
+	const FString OriginalDllPath = FindAgilitySDKFile(TEXT("D3D12StateObjectCompiler.dll"));
+	if (OriginalDllPath.IsEmpty())
+	{
+		UE_LOG(LogASDTool, Error, TEXT("D3D12StateObjectCompiler.dll not found in AgilitySDK directories."));
+		return false;
+	}
+
+	const int32 NumThreads   = (Args.NumThreads > 0) ? Args.NumThreads : FPlatformMisc::NumberOfCoresIncludingHyperthreads();
+	const FString SODBName   = FPaths::GetBaseFilename(DestPSDBPath);
+	const FString TempBaseDir = FPaths::GetPath(DestPSDBPath) / TEXT(".tmp");
+
+	if (IFileManager::Get().DirectoryExists(*TempBaseDir))
+	{
+		IFileManager::Get().DeleteDirectory(*TempBaseDir, false, true);
+	}
+
+	UE_LOG(LogASDTool, Display, TEXT("  Compiling %d PSOs across %d threads (per-thread DLL instances)..."),
+		SODBEntries.Num(), NumThreads);
+
+	TArray<TUniquePtr<FThreadCompilerContext>> Contexts;
+	TArray<FString> ThreadPSDBPaths;
+	Contexts.Reserve(NumThreads);
+	ThreadPSDBPaths.Reserve(NumThreads);
+
+	for (int32 i = 0; i < NumThreads; ++i)
+	{
+		FString ThreadTempDir  = TempBaseDir / FString::Printf(TEXT("thread_%d"), i);
+		FString ThreadPSDBPath = ThreadTempDir / (SODBName + TEXT(".psdb"));
+
+		TUniquePtr<FThreadCompilerContext> Ctx = MakeUnique<FThreadCompilerContext>();
+		if (!Ctx->Initialize(OriginalDllPath, ThreadTempDir))
+		{
+			UE_LOG(LogASDTool, Error, TEXT("Failed to initialize thread compiler context %d -- aborting."), i);
+			return false;
+		}
+		Contexts.Add(MoveTemp(Ctx));
+		ThreadPSDBPaths.Add(MoveTemp(ThreadPSDBPath));
+	}
+
+	// std::atomic used here for fetch_add (old-value-returning increment) which TAtomic doesn't expose.
+	// This is a tool-only code path so std:: usage is acceptable.
+	std::atomic<int32> NextIndex{0};
+	std::atomic<int32> CompiledCount{0};
+	std::atomic<int32> FailedCount{0};
+	std::atomic<int64> TotalCompileCycles{0};
+	FCriticalSection   CreateCompilerLock;
+
+	double StartTime = FPlatformTime::Seconds();
+
+	TArray<TFuture<void>> Futures;
+	Futures.Reserve(NumThreads);
+	for (int32 i = 0; i < NumThreads; ++i)
+	{
+		FThreadCompilerContext* Ctx           = Contexts[i].Get();
+		FString                 ThreadPSDBPath = ThreadPSDBPaths[i];
+		Futures.Add(Async(EAsyncExecution::Thread, [&, Ctx, ThreadPSDBPath]()
+		{
+			CompileOnThread(*Ctx, Args, Plugin, Family, ABIVersion,
+				SODBEntries, ThreadPSDBPath, NextIndex, CompiledCount, FailedCount, TotalCompileCycles, CreateCompilerLock);
+		}));
+	}
+	for (TFuture<void>& Future : Futures)
+	{
+		Future.Wait();
+	}
+
+	// Destroy contexts before merging -- ensures DLL has released file handles to per-thread PSDBs
+	Contexts.Empty();
+
+	double ElapsedSeconds  = FPlatformTime::Seconds() - StartTime;
+	const int32  NumCompiled = (int32)CompiledCount;
+	const int32  NumFailed   = (int32)FailedCount;
+	const double AvgMs       = (NumCompiled > 0) ? FPlatformTime::ToSeconds64((int64)TotalCompileCycles) / NumCompiled * 1000.0 : 0.0;
+
+	UE_LOG(LogASDTool, Display, TEXT("  COM compile done: %d compiled, %d failed, %.1f s | %d threads | avg %.1f ms/PSO"),
+		NumCompiled, NumFailed, ElapsedSeconds, NumThreads, AvgMs);
+
+	TArray<FString> ExistingPSDBPaths;
+	for (const FString& Path : ThreadPSDBPaths)
+	{
+		if (FPaths::FileExists(Path))
+		{
+			ExistingPSDBPaths.Add(Path);
+		}
+	}
+
+	double MergeSeconds = 0.0;
+	const bool bMergeSuccess = MergePSDBs(ExistingPSDBPaths, DestPSDBPath, MergeSeconds);
+
+	if (bMergeSuccess)
+	{
+		IFileManager::Get().DeleteDirectory(*TempBaseDir, false, true);
+	}
+	else
+	{
+		UE_LOG(LogASDTool, Warning, TEXT("  Merge failed -- per-thread PSDBs preserved in: %s"), *TempBaseDir);
+	}
+
+	return bMergeSuccess;
+}
+
+static void CompileWithCOM(
+	const FCompilePSDBsArgs& Args,
+	const TArray<FString>& SODBFiles,
+	const TArray<FCompileWorkItem>& CompileItems,
+	FCompilePSDBsStats& OutStats)
+{
+	UE_LOG(LogASDTool, Display, TEXT("=== Compiling PSDBs (COM) ==="));
+
+	// Pre-load all SODB entries once -- shared across all work items for the same SODB
+	TMap<FString, TArray<FSODBEntry>> SODBEntryCache;
+	for (const FString& SODBFile : SODBFiles)
+	{
+		TArray<FSODBEntry>& Entries = SODBEntryCache.Add(SODBFile);
+		if (!ReadSODBEntries(SODBFile, Entries))
+		{
+			UE_LOG(LogASDTool, Error, TEXT("Failed to read SODB entries from: %s"), *SODBFile);
+			++OutStats.PSDBsFailed;
+		}
+	}
+
+	for (int32 i = 0; i < CompileItems.Num(); ++i)
+	{
+		const FCompileWorkItem& Item = CompileItems[i];
+		UE_LOG(LogASDTool, Display, TEXT("[%d/%d] %s | ABI %llu | %s"),
+			   i + 1, CompileItems.Num(), *Item.Family->Name, Item.ABIVersion, *Item.SODBName);
+
+		const TArray<FSODBEntry>* CachedEntries = SODBEntryCache.Find(Item.SODBPath);
+		if (!CachedEntries || CachedEntries->IsEmpty())
+		{
+			UE_LOG(LogASDTool, Error, TEXT("  No SODB entries available for: %s"), *Item.SODBName);
+			++OutStats.PSDBsFailed;
+			continue;
+		}
+
+		if (CompileDirectCOM(Args, *Item.Plugin, *Item.Family, Item.ABIVersion, *CachedEntries, Item.OutputPath))
+		{
+			++OutStats.PSDBsCompiled;
+		}
+		else
+		{
+			++OutStats.PSDBsFailed;
+		}
+	}
+}
+
+
+//--------------------------------------------------------------------------------------------------
+// Usage / arg parsing
+//--------------------------------------------------------------------------------------------------
+
+static void PrintUsage()
+{
+	UE_LOG(LogASDTool, Display, TEXT("CompilePSDBs: compile State Object Databases (.sodb) into Precompiled Shader Databases (.psdb)"));
+	UE_LOG(LogASDTool, Display, TEXT(""));
+	UE_LOG(LogASDTool, Display, TEXT("Required flags:"));
+	UE_LOG(LogASDTool, Display, TEXT("    -SODBPath=<file_or_dir>   Path to a single .sodb file or directory of .sodb files"));
+	UE_LOG(LogASDTool, Display, TEXT("    -OutputDir=<dir>          Root output directory for .psdb files"));
+	UE_LOG(LogASDTool, Display, TEXT("    -Name=<string>            Application name"));
+	UE_LOG(LogASDTool, Display, TEXT("    -ExeFilename=<string>     Application exe filename"));
+	UE_LOG(LogASDTool, Display, TEXT("    -Engine=<string>          Engine name"));
+	UE_LOG(LogASDTool, Display, TEXT(""));
+	UE_LOG(LogASDTool, Display, TEXT("Optional flags:"));
+	UE_LOG(LogASDTool, Display, TEXT("    -PluginDir=<dir>          Compiler plugins directory (default: Engine/Binaries/ThirdParty/D3D12CompilerPlugins)"));
+	UE_LOG(LogASDTool, Display, TEXT("    -CompilerExe=<path>       Path to D3D12StateObjectCompiler.exe (default: auto-detect)"));
+	UE_LOG(LogASDTool, Display, TEXT("    -LocalAdapter             Only compile for adapter families matching installed GPU(s). Mutually exclusive with -AdapterFamily/-ABI."));
+	UE_LOG(LogASDTool, Display, TEXT("    -ListOnly                 Enumerate plugins and families without compiling"));
+	UE_LOG(LogASDTool, Display, TEXT("    -Clean                    Delete entire output directory before compiling"));
+	UE_LOG(LogASDTool, Display, TEXT("    -Compress                 Kraken compress output PSDBs (.psdb.oodle)"));
+	UE_LOG(LogASDTool, Display, TEXT("    -UseCompilerExe           Use D3D12StateObjectCompiler.exe instead of direct COM compilation"));
+	UE_LOG(LogASDTool, Display, TEXT("    -Threads=<n>              Number of compiler threads for COM path (default: all hardware cores)"));
+	UE_LOG(LogASDTool, Display, TEXT("    -AdapterFamily=<name>     Only compile for this adapter family (e.g. AMD_NAVI48:gfx1201). Mutually exclusive with -LocalAdapter."));
+	UE_LOG(LogASDTool, Display, TEXT("    -ABI=<uint64>              ABI version to use with -AdapterFamily (omit to use DLL-reported ABI). Mutually exclusive with -LocalAdapter."));
+	UE_LOG(LogASDTool, Display, TEXT("    -Help                     Print this help message"));
+}
+
+static bool ParseCompilePSDBsArgs(FCompilePSDBsArgs& OutArgs)
+{
+	const TCHAR* CmdLine = FCommandLine::Get();
+
+	OutArgs.bListOnly       = FParse::Param(CmdLine, TEXT("ListOnly"));
+	OutArgs.bClean          = FParse::Param(CmdLine, TEXT("Clean"));
+	OutArgs.bLocalAdapter   = FParse::Param(CmdLine, TEXT("LocalAdapter"));
+	OutArgs.bCompress       = FParse::Param(CmdLine, TEXT("Compress"));
+	OutArgs.bUseCompilerExe = FParse::Param(CmdLine, TEXT("UseCompilerExe"));
+
+	FString ThreadsStr;
+	if (FParse::Value(CmdLine, TEXT("-Threads="), ThreadsStr))
+	{
+		OutArgs.NumThreads = FCString::Atoi(*ThreadsStr);
+	}
+
+	FParse::Value(CmdLine, TEXT("-AdapterFamily="), OutArgs.AdapterFamily);
+	FString ABIStr;
+	if (FParse::Value(CmdLine, TEXT("-ABI="), ABIStr))
+	{
+		OutArgs.ABI = FCString::Strtoui64(*ABIStr, nullptr, 10);
+	}
+
+	if (OutArgs.bLocalAdapter && !OutArgs.AdapterFamily.IsEmpty())
+	{
+		UE_LOG(LogASDTool, Error, TEXT("-LocalAdapter and -AdapterFamily are mutually exclusive"));
+		return false;
+	}
+
+	FParse::Value(CmdLine, TEXT("-SODBPath="),    OutArgs.SODBPath);
+	FParse::Value(CmdLine, TEXT("-OutputDir="),   OutArgs.OutputDir);
+	FParse::Value(CmdLine, TEXT("-PluginDir="),   OutArgs.PluginDir);
+	FParse::Value(CmdLine, TEXT("-CompilerExe="), OutArgs.CompilerExe);
+	FParse::Value(CmdLine, TEXT("-Name="),        OutArgs.AppName);
+	FParse::Value(CmdLine, TEXT("-ExeFilename="), OutArgs.ExeFilename);
+	FParse::Value(CmdLine, TEXT("-Engine="),      OutArgs.EngineName);
+
+	if (OutArgs.CompilerExe.IsEmpty())
+	{
+		OutArgs.CompilerExe = FindAgilitySDKFile(TEXT("D3D12StateObjectCompiler.exe"));
+		if (OutArgs.CompilerExe.IsEmpty())
+		{
+			UE_LOG(LogASDTool, Error, TEXT("D3D12StateObjectCompiler.exe not found. Use -CompilerExe= to specify path."));
+			return false;
+		}
+	}
+	UE_LOG(LogASDTool, Display, TEXT("Compiler exe: %s"), *OutArgs.CompilerExe);
+	UE_LOG(LogASDTool, Display, TEXT("Compile path: %s"), OutArgs.bUseCompilerExe ? TEXT("EXE") : TEXT("COM (direct)"));
+
+	if (OutArgs.PluginDir.IsEmpty())
+	{
+		OutArgs.PluginDir = FindDefaultCompilerPluginDir();
+		if (OutArgs.PluginDir.IsEmpty())
+		{
+			UE_LOG(LogASDTool, Error, TEXT("Plugin directory not found. Use -PluginDir= to specify path."));
+			return false;
+		}
+	}
+	UE_LOG(LogASDTool, Display, TEXT("Plugin dir:   %s"), *OutArgs.PluginDir);
+
+	if (!OutArgs.bListOnly)
+	{
+		bool bMissingArgs = OutArgs.SODBPath.IsEmpty() || OutArgs.OutputDir.IsEmpty() ||
+			OutArgs.AppName.IsEmpty() || OutArgs.ExeFilename.IsEmpty() || OutArgs.EngineName.IsEmpty();
+		if (bMissingArgs)
+		{
+			PrintUsage();
+			if (OutArgs.SODBPath.IsEmpty())
+			{
+				UE_LOG(LogASDTool, Error, TEXT("Missing -SODBPath"));
+			}
+			if (OutArgs.OutputDir.IsEmpty())
+			{
+				UE_LOG(LogASDTool, Error, TEXT("Missing -OutputDir"));
+			}
+			if (OutArgs.AppName.IsEmpty())
+			{
+				UE_LOG(LogASDTool, Error, TEXT("Missing -Name"));
+			}
+			if (OutArgs.ExeFilename.IsEmpty())
+			{
+				UE_LOG(LogASDTool, Error, TEXT("Missing -ExeFilename"));
+			}
+			if (OutArgs.EngineName.IsEmpty())
+			{
+				UE_LOG(LogASDTool, Error, TEXT("Missing -Engine"));
+			}
+			return false;
+		}
+	}
+	return true;
+}
+
+//--------------------------------------------------------------------------------------------------
+// Step helper functions
+//--------------------------------------------------------------------------------------------------
+
+static int32 CollectSODBFiles(const FCompilePSDBsArgs& Args, TArray<FString>& OutSODBFiles)
+{
+	if (FPaths::FileExists(Args.SODBPath))
+	{
+		OutSODBFiles.Add(FPaths::ConvertRelativePathToFull(Args.SODBPath));
+	}
+	else if (FPaths::DirectoryExists(Args.SODBPath))
+	{
+		IFileManager::Get().FindFiles(OutSODBFiles, *Args.SODBPath, TEXT(".sodb"));
+		for (FString& File : OutSODBFiles)
+		{
+			File = FPaths::ConvertRelativePathToFull(Args.SODBPath / File);
+		}
+	}
+	else
+	{
+		UE_LOG(LogASDTool, Error, TEXT("SODBPath '%s' is not a valid file or directory"), *Args.SODBPath);
+		return 1;
+	}
+
+	if (OutSODBFiles.IsEmpty())
+	{
+		UE_LOG(LogASDTool, Error, TEXT("No .sodb files found at '%s'"), *Args.SODBPath);
+		return 1;
+	}
+
+	UE_LOG(LogASDTool, Display, TEXT("Found %d SODB file(s)"), OutSODBFiles.Num());
+	return 0;
+}
+
+static void PrepareOutputDirectory(const FCompilePSDBsArgs& Args)
+{
+	if (Args.bClean && FPaths::DirectoryExists(Args.OutputDir))
+	{
+		UE_LOG(LogASDTool, Display, TEXT("Cleaning output directory: %s"), *Args.OutputDir);
+		IFileManager::Get().DeleteDirectory(*Args.OutputDir, false, true);
+	}
+	IFileManager::Get().MakeDirectory(*Args.OutputDir, true);
+}
+
+static void BuildCompileWorkItems(
+	const FCompilePSDBsArgs& Args,
+	const TArray<FCompilerPlugin>& Plugins,
+	const TArray<FLocalAdapterInfo>& LocalAdapters,
+	const TArray<FString>& SODBFiles,
+	TArray<FCompileWorkItem>& OutItems)
+{
+	for (const FCompilerPlugin& Plugin : Plugins)
+	{
+		for (const FAdapterFamily& Family : Plugin.Families)
+		{
+			if (!Args.AdapterFamily.IsEmpty())
+			{
+				// -AdapterFamily/-ABI filter: skip families that don't match the requested name.
+				// If -ABI is also specified, additionally require the plugin reports that exact ABI
+				// so only the one correct plugin DLL is selected, not every DLL that covers the family.
+				if (Family.Name != Args.AdapterFamily)
+				{
+					continue;
+				}
+				if (Args.ABI != 0 && !Family.ABIVersions.Contains(Args.ABI))
+				{
+					continue;
+				}
+			}
+			else if (Args.bLocalAdapter)
+			{
+				bool bMatchesLocal = false;
+				for (const FLocalAdapterInfo& LocalAdapter : LocalAdapters)
+				{
+					if (Family.Name == LocalAdapter.FamilyName)
+					{
+						for (uint64 PluginABI : Family.ABIVersions)
+						{
+							if (PluginABI >= LocalAdapter.MinABIVersion && PluginABI <= LocalAdapter.MaxABIVersion)
+							{
+								bMatchesLocal = true;
+								break;
+							}
+						}
+					}
+				}
+				if (!bMatchesLocal)
+				{
+					continue;
+				}
+			}
+
+			// Use forced ABI if specified, otherwise iterate all ABIs the plugin reports for this family
+			TArray<uint64> ABIsToCompile;
+			if (!Args.AdapterFamily.IsEmpty() && Args.ABI != 0)
+			{
+				ABIsToCompile.Add(Args.ABI);
+			}
+			else
+			{
+				ABIsToCompile = Family.ABIVersions;
+			}
+			for (uint64 ABI : ABIsToCompile)
+			{
+				FString SafeFamilyName = Family.Name;
+				SafeFamilyName.ReplaceCharInline(TCHAR(':'), TCHAR('_'));
+
+				for (const FString& SODBFile : SODBFiles)
+				{
+					FString SODBName = FPaths::GetBaseFilename(SODBFile);
+					FCompileWorkItem Item;
+					Item.Plugin     = &Plugin;
+					Item.Family     = &Family;
+					Item.ABIVersion = ABI;
+					Item.SODBPath   = SODBFile;
+					Item.SODBName   = SODBName;
+					Item.OutputPath = Args.OutputDir / SafeFamilyName / FString::Printf(TEXT("abi_%llu"), ABI) / (SODBName + TEXT(".psdb"));
+					OutItems.Add(MoveTemp(Item));
+				}
+			}
+		}
+	}
+}
+
+static void CompressAllPSDBs(const FCompilePSDBsArgs& Args, FCompilePSDBsStats& OutStats)
+{
+	UE_LOG(LogASDTool, Display, TEXT("=== Compressing PSDBs with Kraken ==="));
+
+	TArray<FString> PSDBFiles;
+	IFileManager::Get().FindFilesRecursive(PSDBFiles, *Args.OutputDir, TEXT("*.psdb"), true, false);
+
+	for (const FString& PSDBFile : PSDBFiles)
+	{
+		if (CompressFile(PSDBFile, PSDB_COMPRESSED_MAGIC))
+		{
+			++OutStats.PSDBsCompressed;
+		}
+		else
+		{
+			++OutStats.CompressFailed;
+		}
+	}
+}
+
+static void PrintStats(
+	const FCompilePSDBsArgs& Args,
+	const TArray<FCompilerPlugin>& Plugins,
+	int32 TotalFamilies,
+	const TArray<FString>& SODBFiles,
+	const FCompilePSDBsStats& Stats)
+{
+	UE_LOG(LogASDTool, Display, TEXT(""));
+	UE_LOG(LogASDTool, Display, TEXT("=== CompilePSDBs Summary ==="));
+	UE_LOG(LogASDTool, Display, TEXT("  Plugins:          %d  (%d adapter families)"), Plugins.Num(), TotalFamilies);
+	UE_LOG(LogASDTool, Display, TEXT("  SODB files:       %d"), SODBFiles.Num());
+	UE_LOG(LogASDTool, Display, TEXT("  PSDBs compiled:   %d  (%d failed)"), Stats.PSDBsCompiled, Stats.PSDBsFailed);
+	if (Args.bCompress)
+	{
+		UE_LOG(LogASDTool, Display, TEXT("  PSDBs compressed: %d  (%d failed)"), Stats.PSDBsCompressed, Stats.CompressFailed);
+	}
+	UE_LOG(LogASDTool, Display, TEXT("  Total time:       %.1f s"), Stats.ElapsedSeconds);
+	UE_LOG(LogASDTool, Display, TEXT("  Output:           %s"), *Args.OutputDir);
+}
+
+//--------------------------------------------------------------------------------------------------
+// CompilePSDBs -- entry point
+//
+// Steps:
+//   Step 0:  Parse and validate command-line arguments
+//   Step 1:  Discover and enumerate compiler plugins
+//            [if -ListOnly: print and return]
+//   Step 2:  Filter adapters  [if -LocalAdapter or -AdapterFamily/-ABI]
+//   Step 3:  Collect SODB files
+//   Step 4:  Prepare output directory  [clean if -Clean, then create]
+//   Step 5:  Build compile work items
+//   Step 6a: Compile PSDBs via EXE  [if -UseCompilerExe]
+//   Step 6b: Compile PSDBs via COM  [default]
+//   Step 7:  Compress PSDBs  [if -Compress]
+//   Step 8:  Print summary
+//--------------------------------------------------------------------------------------------------
+
+int32 CompilePSDBs()
+{
+	const TCHAR* CmdLine = FCommandLine::Get();
+
+	if (FParse::Param(CmdLine, TEXT("h")) || FParse::Param(CmdLine, TEXT("help")))
+	{
+		PrintUsage();
+		return 0;
+	}
+
+	uint64 TotalStartCycles = FPlatformTime::Cycles64();
+
+	// -------------------------------------------------------------------------
+	// Step 0: Parse and validate command-line arguments
+	// -------------------------------------------------------------------------
+
+	FCompilePSDBsArgs Args;
+	if (!ParseCompilePSDBsArgs(Args))
+	{
+		return 1;
+	}
+
+	// -------------------------------------------------------------------------
+	// Step 1: Discover and enumerate compiler plugins
+	// -------------------------------------------------------------------------
+
+	TArray<FCompilerPlugin> Plugins;
+	int32 TotalFamilies = 0;
+
+	int32 Result = DiscoverAndEnumeratePlugins(Args, Plugins, TotalFamilies);
+	if (Result != 0)
+	{
+		return Result;
+	}
+	if (Args.bListOnly)
+	{
+		return 0;
+	}
+
+	// -------------------------------------------------------------------------
+	// Step 2: Filter adapters  [if -LocalAdapter or -AdapterFamily/-ABI]
+	// -------------------------------------------------------------------------
+
+	TArray<FLocalAdapterInfo> LocalAdapters;
+	if (Args.bLocalAdapter)
+	{
+		Result = FilterLocalAdapters(Args, LocalAdapters);
+		if (Result != 0)
+		{
+			return Result;
+		}
+	}
+
+	// -------------------------------------------------------------------------
+	// Step 3: Collect SODB files
+	// -------------------------------------------------------------------------
+
+	TArray<FString> SODBFiles;
+	Result = CollectSODBFiles(Args, SODBFiles);
+	if (Result != 0)
+	{
+		return Result;
+	}
+
+	// -------------------------------------------------------------------------
+	// Step 4: Prepare output directory
+	// -------------------------------------------------------------------------
+
+	PrepareOutputDirectory(Args);
+
+	// -------------------------------------------------------------------------
+	// Step 5: Build compile work items
+	// -------------------------------------------------------------------------
+
+	TArray<FCompileWorkItem> CompileItems;
+	BuildCompileWorkItems(Args, Plugins, LocalAdapters, SODBFiles, CompileItems);
+
+	// -------------------------------------------------------------------------
+	// Step 6a/6b: Compile PSDBs
+	// -------------------------------------------------------------------------
+
+	FCompilePSDBsStats Stats;
+
+	if (Args.bUseCompilerExe)
+	{
+		CompileWithExe(Args, CompileItems, Stats);
+	}
+	else
+	{
+		CompileWithCOM(Args, SODBFiles, CompileItems, Stats);
+	}
+
+	// -------------------------------------------------------------------------
+	// Step 7: Compress PSDBs  [if -Compress]
+	// -------------------------------------------------------------------------
+
+	if (Args.bCompress)
+	{
+		CompressAllPSDBs(Args, Stats);
+	}
+
+	// -------------------------------------------------------------------------
+	// Step 8: Print summary
+	// -------------------------------------------------------------------------
+
+	Stats.ElapsedSeconds = FPlatformTime::ToSeconds64(FPlatformTime::Cycles64() - TotalStartCycles);
+	PrintStats(Args, Plugins, TotalFamilies, SODBFiles, Stats);
+
+	return (Stats.PSDBsFailed + Stats.CompressFailed > 0) ? 1 : 0;
+}
+
+}; // namespace ASDTool
diff --git a/Engine/Source/Programs/Windows/ASDTool/Private/GenerateSODB.cpp b/Engine/Source/Programs/Windows/ASDTool/Private/GenerateSODB.cpp
new file mode 100644
--- /dev/null
+++ b/Engine/Source/Programs/Windows/ASDTool/Private/GenerateSODB.cpp
@@ -0,0 +1,997 @@
+// Copyright Epic Games, Inc. All Rights Reserved.
+
+#include "ASDToolCommands.h"
+
+// initguid.h must be included before any D3D12 headers to instantiate GUIDs
+#include "Windows/AllowWindowsPlatformTypes.h"
+THIRD_PARTY_INCLUDES_START
+#include <initguid.h>
+#include <D3D12.h>
+THIRD_PARTY_INCLUDES_END
+#include "Windows/HideWindowsPlatformTypes.h"
+
+#include "SODBFile.h"
+#include "ASDToolShaderUtils.h"
+#include "D3D12ShaderUtils.h"
+
+#include "Misc/CommandLine.h"
+#include "Misc/Parse.h"
+#include "Misc/Paths.h"
+#include "HAL/FileManager.h"
+#include "Misc/ScopeRWLock.h"
+#include "HAL/PlatformTime.h"
+#include "HAL/PlatformAtomics.h"
+#include "PipelineCacheUtilities.h"
+#include "PipelineFileCache.h"
+#include "ShaderCodeLibrary.h"
+#include "Misc/EngineVersion.h"
+
+#include "IncludeSQLite.h"
+
+namespace ASDTool
+{
+
+
+//--------------------------------------------------------------------------------------------------
+// Command-line arguments
+//--------------------------------------------------------------------------------------------------
+
+struct FGenerateSODBArgs
+{
+	FString ShaderCodeDir;
+	FString ShkDir;
+	FString BundleDir;      // Directory of .upipelinecache files (enables graphics PSO support)
+	FString OutputDir;
+	FString ShaderTypeFilterStr;
+	FString AppName;
+	FString ExeFilename;
+	FString EngineName;
+	bool bAllCompute     = false;  // Include all compute shaders from bytecode archives (no SHK required)
+	bool bSingleThreaded = false;
+	bool bCleanOutput    = false;
+	bool bSplitSODB      = false;
+	bool bDirectWrite    = false;
+};
+
+//--------------------------------------------------------------------------------------------------
+// FCachedShaderEntry / FShaderCacheMap
+// Populated by PreloadShaderCache. Maps every SM6 shader hash to its raw DXBC and full
+// uncompressed blob so that CollectComputePSOs and CollectPipelineFileCachePSOs can resolve
+// shaders without re-reading archives.
+//--------------------------------------------------------------------------------------------------
+
+struct FCachedShaderEntry
+{
+	EShaderFrequency Frequency = SF_NumFrequencies;
+	TArray<uint8> DXBCData;    // Raw DXBC (ShaderCode view -- after SRT)
+	TArray<uint8> FullBlob;    // Full uncompressed UE shader blob (for root sig extraction)
+};
+
+using FShaderCacheMap = TMap<FShaderHash, FCachedShaderEntry>;
+
+//--------------------------------------------------------------------------------------------------
+// FGenerateSODBStats
+// Accumulated across CollectComputePSOs and CollectPipelineFileCachePSOs.
+// ElapsedSeconds covers total wall time and is set in GenerateSODBs().
+//--------------------------------------------------------------------------------------------------
+
+struct FGenerateSODBStats
+{
+	int32  ComputePSOsWritten  = 0;
+	int32  ComputePSOsFailed   = 0;
+	int32  GraphicsPSOsWritten = 0;
+	int32  GraphicsPSOsFailed  = 0;
+	double ElapsedSeconds      = 0.0;  // Total wall time, set in GenerateSODBs()
+};
+
+//--------------------------------------------------------------------------------------------------
+// Usage / arg parsing
+//--------------------------------------------------------------------------------------------------
+
+static void PrintUsage()
+{
+	UE_LOG(LogASDTool, Display, TEXT("GenerateSODB: generate State Object Database (.sodb) files for compute and graphics shaders"));
+	UE_LOG(LogASDTool, Display, TEXT("  Extracts shader DXBC bytecode from cooked shader archives and writes"));
+	UE_LOG(LogASDTool, Display, TEXT("  them into SODB files for DirectX Advanced Shader Delivery."));
+	UE_LOG(LogASDTool, Display, TEXT(""));
+	UE_LOG(LogASDTool, Display, TEXT("Required flags:"));
+	UE_LOG(LogASDTool, Display, TEXT("    -ShaderCodeDir=<dir>: directory where the .ushaderbytecode files are located"));
+	UE_LOG(LogASDTool, Display, TEXT("    -OutputDir=<dir>: directory where the .sodb files will be created"));
+	UE_LOG(LogASDTool, Display, TEXT("    -Name=<string>: application name (e.g. MyGame)"));
+	UE_LOG(LogASDTool, Display, TEXT("    -ExeFilename=<string>: application exe filename (e.g. MyGame-Win64-Test.exe)"));
+	UE_LOG(LogASDTool, Display, TEXT("    -Engine=<string>: engine name (e.g. UnrealEngine5.8)"));
+	UE_LOG(LogASDTool, Display, TEXT("At least one collection source (required):"));
+	UE_LOG(LogASDTool, Display, TEXT("    -AllCompute: collect all compute shaders from bytecode archives"));
+	UE_LOG(LogASDTool, Display, TEXT("    -BundleDir=<dir>: directory of .upipelinecache files (graphics + compute PSOs)"));
+	UE_LOG(LogASDTool, Display, TEXT("Optional grouping:"));
+	UE_LOG(LogASDTool, Display, TEXT("    -ShkDir=<dir>: directory of SHK files -- routes compute shaders into named groups"));
+	UE_LOG(LogASDTool, Display, TEXT("Optional filtering flags:"));
+	UE_LOG(LogASDTool, Display, TEXT("    -ShaderType=<type1,type2,...>: only include specific shader types (e.g. FLumenCardCS,Global)"));
+	UE_LOG(LogASDTool, Display, TEXT("Optional misc flags:"));
+	UE_LOG(LogASDTool, Display, TEXT("    -OneThread: don't parallelize extraction"));
+	UE_LOG(LogASDTool, Display, TEXT("    -Clean: delete existing .sodb files in the output directory before generating"));
+	UE_LOG(LogASDTool, Display, TEXT("    -DirectWrite: write SODB via direct SQLite instead of D3D12 API (faster)"));
+	UE_LOG(LogASDTool, Display, TEXT("    -SplitSODB: create separate SODB files per shader type (default: single file)"));
+	UE_LOG(LogASDTool, Display, TEXT("    -Help: print this help message"));
+}
+
+static bool ParseGenerateSODBArgs(FGenerateSODBArgs& OutArgs)
+{
+	const TCHAR* CmdLine = FCommandLine::Get();
+
+	OutArgs.bAllCompute     = FParse::Param(CmdLine, TEXT("AllCompute"));
+	OutArgs.bSingleThreaded = FParse::Param(CmdLine, TEXT("OneThread"));
+	OutArgs.bCleanOutput    = FParse::Param(CmdLine, TEXT("Clean"));
+	OutArgs.bSplitSODB      = FParse::Param(CmdLine, TEXT("SplitSODB"));
+	OutArgs.bDirectWrite    = FParse::Param(CmdLine, TEXT("DirectWrite"));
+
+	FParse::Value(CmdLine, TEXT("-ShaderCodeDir="), OutArgs.ShaderCodeDir);
+	FParse::Value(CmdLine, TEXT("-ShkDir="),        OutArgs.ShkDir);
+	FParse::Value(CmdLine, TEXT("-BundleDir="),      OutArgs.BundleDir);
+	FParse::Value(CmdLine, TEXT("-OutputDir="),      OutArgs.OutputDir);
+	FParse::Value(CmdLine, TEXT("-ShaderType="),     OutArgs.ShaderTypeFilterStr, false);
+	FParse::Value(CmdLine, TEXT("-Name="),           OutArgs.AppName);
+	FParse::Value(CmdLine, TEXT("-ExeFilename="),    OutArgs.ExeFilename);
+	FParse::Value(CmdLine, TEXT("-Engine="),         OutArgs.EngineName);
+
+	// At least one collection source is required -- ShkDir alone only sets up group routing
+	// but does not collect anything without -AllCompute or -BundleDir.
+	bool bHasCollectionSource = OutArgs.bAllCompute || !OutArgs.BundleDir.IsEmpty();
+
+	if (OutArgs.OutputDir.IsEmpty() || OutArgs.ShaderCodeDir.IsEmpty() || !bHasCollectionSource ||
+		OutArgs.AppName.IsEmpty() || OutArgs.ExeFilename.IsEmpty() || OutArgs.EngineName.IsEmpty())
+	{
+		PrintUsage();
+		if (OutArgs.OutputDir.IsEmpty())        { UE_LOG(LogASDTool, Error, TEXT("No output directory specified.")); }
+		if (OutArgs.ShaderCodeDir.IsEmpty())    { UE_LOG(LogASDTool, Error, TEXT("No shader code directory specified.")); }
+		if (!bHasCollectionSource)              { UE_LOG(LogASDTool, Error, TEXT("Specify at least one of: -AllCompute, -BundleDir (-ShkDir alone only sets up group routing)")); }
+		if (OutArgs.AppName.IsEmpty())          { UE_LOG(LogASDTool, Error, TEXT("Missing -Name")); }
+		if (OutArgs.ExeFilename.IsEmpty())      { UE_LOG(LogASDTool, Error, TEXT("Missing -ExeFilename")); }
+		if (OutArgs.EngineName.IsEmpty())       { UE_LOG(LogASDTool, Error, TEXT("Missing -Engine")); }
+		return false;
+	}
+
+	if ((OutArgs.bSplitSODB || !OutArgs.ShaderTypeFilterStr.IsEmpty()) && OutArgs.ShkDir.IsEmpty())
+	{
+		UE_LOG(LogASDTool, Error, TEXT("-SplitSODB and -ShaderType require -ShkDir"));
+		return false;
+	}
+
+	return true;
+}
+
+//--------------------------------------------------------------------------------------------------
+// PrepareOutputDirectory
+// Deletes stale .sodb files from the output directory if -Clean was requested, then creates it.
+//--------------------------------------------------------------------------------------------------
+
+static int32 PrepareOutputDirectory(const FGenerateSODBArgs& Args)
+{
+	if (Args.bCleanOutput)
+	{
+		TArray<FString> ExistingSODBs;
+		IFileManager::Get().FindFiles(ExistingSODBs, *Args.OutputDir, TEXT(".sodb"));
+		for (const FString& Existing : ExistingSODBs)
+		{
+			const FString FullPath = Args.OutputDir / Existing;
+			if (!IFileManager::Get().Delete(*FullPath))
+			{
+				UE_LOG(LogASDTool, Warning, TEXT("Failed to delete: %s"), *FullPath);
+			}
+		}
+		if (ExistingSODBs.Num() > 0)
+		{
+			UE_LOG(LogASDTool, Display, TEXT("Cleaned %d existing .sodb files."), ExistingSODBs.Num());
+		}
+	}
+
+	if (!IFileManager::Get().MakeDirectory(*Args.OutputDir, true))
+	{
+		UE_LOG(LogASDTool, Error, TEXT("Cannot create directory '%s'."), *Args.OutputDir);
+		return 2;
+	}
+
+	return 0;
+}
+
+//--------------------------------------------------------------------------------------------------
+// SetupGroups
+// If -ShkDir is provided, loads SHK files to build named groups with their wanted compute shader
+// hash sets (and optionally a shader type filter / split-by-type). Otherwise creates a single
+// "default" group that will accept all compute shaders (-AllCompute) or graphics PSOs (-BundleDir).
+//--------------------------------------------------------------------------------------------------
+
+static int32 SetupGroups(
+	const FGenerateSODBArgs& Args,
+	TMap<FString, TUniquePtr<FSODBGroup>>& OutGroups)
+{
+	if (!Args.ShkDir.IsEmpty())
+	{
+		// Parse shader type filter
+		TSet<FString> ShaderTypeFilter;
+		bool bFilterGlobals = false;
+		if (!Args.ShaderTypeFilterStr.IsEmpty())
+		{
+			TArray<FString> TypeParts;
+			Args.ShaderTypeFilterStr.ParseIntoArray(TypeParts, TEXT(","));
+			for (const FString& Part : TypeParts)
+			{
+				FString Trimmed = Part.TrimStartAndEnd();
+				if (Trimmed.Equals(TEXT("Global"), ESearchCase::IgnoreCase))
+				{
+					bFilterGlobals = true;
+				}
+				else
+				{
+					ShaderTypeFilter.Add(Trimmed);
+				}
+			}
+			UE_LOG(LogASDTool, Display, TEXT("Shader type filter: %s (globals: %s)"), *Args.ShaderTypeFilterStr, bFilterGlobals ? TEXT("yes") : TEXT("no"));
+		}
+		bool bHasShaderTypeFilter = !ShaderTypeFilter.IsEmpty() || bFilterGlobals;
+
+		// Auto-enable SplitSODB when ShaderType filter is specified - otherwise filtered types
+		// all end up in a single "default.sodb" which defeats the purpose of the filter.
+		bool bSplitSODB = Args.bSplitSODB || bHasShaderTypeFilter;
+
+		static const FName NAME_SF_Compute(TEXT("SF_Compute"));
+
+		TArray<FString> AllShkFiles;
+		IFileManager::Get().FindFiles(AllShkFiles, *Args.ShkDir, TEXT(".shk"));
+
+		for (const FString& File : AllShkFiles)
+		{
+			if (!File.Contains(TEXT("SM6")) || !File.Contains(TEXT("PCD3D")))
+			{
+				continue;
+			}
+
+			const FString ShkFilePath = Args.ShkDir / File;
+			bool bIsGlobal = File.Contains(TEXT("Global"));
+
+			TArray<FStableShaderKeyAndValue> Entries;
+			if (!UE::PipelineCacheUtilities::LoadStableKeysFile(ShkFilePath, Entries))
+			{
+				UE_LOG(LogASDTool, Warning, TEXT("Failed to load SHK '%s'."), *ShkFilePath);
+				continue;
+			}
+
+			for (const FStableShaderKeyAndValue& Entry : Entries)
+			{
+				if (Entry.TargetFrequency != NAME_SF_Compute)
+				{
+					continue;
+				}
+
+				// Apply shader type filter
+				if (bHasShaderTypeFilter)
+				{
+					if (bIsGlobal)
+					{
+						if (!bFilterGlobals) continue;
+					}
+					else
+					{
+						if (!ShaderTypeFilter.Contains(Entry.ShaderType.ToString())) continue;
+					}
+				}
+
+				FString GroupName = bSplitSODB
+					? (bIsGlobal ? TEXT("Global") : Entry.ShaderType.ToString())
+					: TEXT("default");
+
+				TUniquePtr<FSODBGroup>& Group = OutGroups.FindOrAdd(GroupName);
+				if (!Group.IsValid())
+				{
+					Group = MakeUnique<FSODBGroup>();
+					Group->Name = GroupName;
+				}
+
+				Group->WantedHashes.Add(Entry.OutputHash);
+			}
+		}
+
+		if (OutGroups.IsEmpty())
+		{
+			UE_LOG(LogASDTool, Error, TEXT("No SM6 compute shader entries found in SHK files."));
+			return 5;
+		}
+
+		int32 TotalWanted = 0;
+		for (auto& Pair : OutGroups)
+		{
+			TotalWanted += Pair.Value->WantedHashes.Num();
+		}
+		UE_LOG(LogASDTool, Display, TEXT("Set up %d SODB group(s), %d wanted compute shaders"), OutGroups.Num(), TotalWanted);
+	}
+	else
+	{
+		// No SHK data -- create a single "default" group that accepts all compute shaders
+		// (used by -AllCompute) and/or all graphics PSOs (used by -BundleDir).
+		TUniquePtr<FSODBGroup>& DefaultGroup = OutGroups.FindOrAdd(TEXT("default"));
+		if (!DefaultGroup.IsValid())
+		{
+			DefaultGroup = MakeUnique<FSODBGroup>();
+			DefaultGroup->Name = TEXT("default");
+		}
+	}
+
+	return 0;
+}
+
+//--------------------------------------------------------------------------------------------------
+// PreloadShaderCache
+// Reads all SM6 .ushaderbytecode archives and builds a map of every shader hash to its raw DXBC
+// and full uncompressed blob. Iterates files serially but processes entries within each file in
+// parallel, since the number of entries per archive vastly outnumbers the number of archive files.
+//--------------------------------------------------------------------------------------------------
+
+static void PreloadShaderCache(
+	const FGenerateSODBArgs& Args,
+	FShaderCacheMap& OutCache)
+{
+	uint64 StartCycles = FPlatformTime::Cycles64();
+
+	TArray<FString> AllFiles;
+	IFileManager::Get().FindFiles(AllFiles, *Args.ShaderCodeDir, TEXT(".ushaderbytecode"));
+
+	TArray<FString> SM6Files;
+	for (const FString& File : AllFiles)
+	{
+		if (File.Contains(TEXT("SM6")) && File.Contains(TEXT("PCD3D")))
+		{
+			SM6Files.Add(File);
+		}
+	}
+
+	// Build a key list from all archive hashes so ProcessArchive will visit everything.
+	// We read archive headers only (no decompression) to enumerate hashes.
+	TArray<ASDTool::FShaderKeyEntry> AllKeys;
+	TSet<FShaderHash> KeysSeen;
+
+	for (const FString& File : SM6Files)
+	{
+		const FString CodeFilePath = Args.ShaderCodeDir / File;
+		TUniquePtr<FArchive> LibraryAr(IFileManager::Get().CreateFileReader(*CodeFilePath));
+		if (!LibraryAr)
+		{
+			continue;
+		}
+
+		uint32 Version = 0;
+		*LibraryAr << Version;
+
+		FSerializedShaderArchive SerializedShaders;
+		*LibraryAr << SerializedShaders;
+
+		for (const FShaderHash& Hash : SerializedShaders.GetShaderHashes())
+		{
+			bool bAlreadyAdded = false;
+			KeysSeen.Add(Hash, &bAlreadyAdded);
+			if (!bAlreadyAdded)
+			{
+				ASDTool::FShaderKeyEntry& Key = AllKeys.Emplace_GetRef();
+				Key.Hash = Hash;
+			}
+		}
+	}
+
+	// Sort keys for merge-join in ProcessArchive
+	AllKeys.Sort([](const ASDTool::FShaderKeyEntry& A, const ASDTool::FShaderKeyEntry& B)
+	{
+		return A.Hash < B.Hash;
+	});
+
+	// Iterate files serially, process entries within each file in parallel.
+	// Entry count per archive vastly outnumbers the number of archive files.
+	FCriticalSection CacheLock;
+	TAtomic<int32> TotalCached{0};
+
+	for (const FString& File : SM6Files)
+	{
+		const FString CodeFilePath = Args.ShaderCodeDir / File;
+		TSet<FShaderHash> LocalSeen;
+		ASDTool::ProcessArchive(CodeFilePath, AllKeys, LocalSeen,
+			[&](ASDTool::FShaderArchiveEntry& Entry)
+			{
+				FCachedShaderEntry CacheEntry;
+				CacheEntry.Frequency = Entry.Frequency;
+				CacheEntry.DXBCData.Append(Entry.ShaderCode.GetData(), Entry.ShaderCode.Num());
+				if (Entry.UncompressedBuffer)
+				{
+					CacheEntry.FullBlob = *Entry.UncompressedBuffer;
+				}
+				{
+					FScopeLock Lock(&CacheLock);
+					OutCache.Emplace(Entry.Hash, MoveTemp(CacheEntry));
+				}
+				++TotalCached;
+			},
+			/*bSingleThreaded=*/Args.bSingleThreaded);  // Respect -OneThread flag
+	}
+
+	double LoadSeconds = FPlatformTime::ToSeconds64(FPlatformTime::Cycles64() - StartCycles);
+	UE_LOG(LogASDTool, Display, TEXT("  Cached %d shaders from %d SM6 archives (%.2f s)"), (int32)TotalCached, SM6Files.Num(), LoadSeconds);
+}
+
+//--------------------------------------------------------------------------------------------------
+// CollectComputePSOs
+// Iterates the pre-loaded shader cache and routes each SM6 compute shader into the correct SODB
+// group. Groups with a wanted hash set (from -ShkDir) match by hash. With -AllCompute (no SHK)
+// all compute shaders go into the single "default" group.
+//--------------------------------------------------------------------------------------------------
+
+static int32 CollectComputePSOs(
+	const FGenerateSODBArgs& Args,
+	const FShaderCacheMap& ShaderCache,
+	TMap<FString, TUniquePtr<FSODBGroup>>& Groups,
+	FGenerateSODBStats& OutStats)
+{
+	if (ShaderCache.IsEmpty())
+	{
+		UE_LOG(LogASDTool, Error, TEXT("Shader cache is empty -- nothing to collect compute PSOs from."));
+		return 4;
+	}
+
+	// Collect group pointers for fast iteration
+	TArray<FSODBGroup*> GroupPtrs;
+	for (auto& Pair : Groups)
+	{
+		GroupPtrs.Add(Pair.Value.Get());
+	}
+
+	bool bAllCompute = Args.bAllCompute;
+	uint64 StartCycles = FPlatformTime::Cycles64();
+
+	for (auto& CachePair : ShaderCache)
+	{
+		const FShaderHash&        Hash   = CachePair.Key;
+		const FCachedShaderEntry& Cached = CachePair.Value;
+
+		if (Cached.Frequency != SF_Compute)
+		{
+			continue;
+		}
+
+		// Verify DXBC container magic
+		uint32 DXBCMagic = 0;
+		if (Cached.DXBCData.Num() >= 4) { FMemory::Memcpy(&DXBCMagic, Cached.DXBCData.GetData(), sizeof(uint32)); }
+		if (Cached.DXBCData.Num() < 4 || DXBCMagic != 0x43425844)
+		{
+			++OutStats.ComputePSOsFailed;
+			continue;
+		}
+
+		// Build root signature from the full shader blob
+		TRefCountPtr<ID3DBlob> RootSigBlob;
+		{
+			TConstArrayView<uint8> FullBlob(Cached.FullBlob.GetData(), Cached.FullBlob.Num());
+			int32 NativeBytecodeOffset = 0;
+			if (!BuildComputeShaderRootSignature(FullBlob, D3D12_RESOURCE_BINDING_TIER_3, RootSigBlob, NativeBytecodeOffset))
+			{
+				UE_LOG(LogASDTool, Warning, TEXT("Failed to build root signature for compute shader %s"), *Hash.ToString());
+				++OutStats.ComputePSOsFailed;
+				continue;
+			}
+		}
+
+		// Route to the correct group
+		bool bFound = false;
+		for (FSODBGroup* Group : GroupPtrs)
+		{
+			if (Group->Contains(Hash))
+			{
+				if (Group->StoreComputeShader(Hash, Cached.DXBCData.GetData(), Cached.DXBCData.Num(), RootSigBlob.GetReference()))
+				{
+					++OutStats.ComputePSOsWritten;
+				}
+				bFound = true;
+				break;
+			}
+		}
+
+		// Route unmatched shaders to the "default" group if one exists.
+		// The "default" group is only present when -ShkDir is not active -- when SHK groups
+		// are set up, all groups are named and there is no "default", so Find returns null
+		// and routing is skipped. TMap iteration order is non-deterministic so we look up
+		// by name rather than by index.
+		if (!bFound && bAllCompute)
+		{
+			TUniquePtr<FSODBGroup>* DefaultGroup = Groups.Find(TEXT("default"));
+			if (DefaultGroup && DefaultGroup->IsValid())
+			{
+				if ((*DefaultGroup)->StoreComputeShader(Hash, Cached.DXBCData.GetData(), Cached.DXBCData.Num(), RootSigBlob.GetReference()))
+				{
+					++OutStats.ComputePSOsWritten;
+				}
+			}
+		}
+	}
+
+	UE_LOG(LogASDTool, Display, TEXT("  Compute PSOs: %d written, %d failed (%.2f s)"),
+		OutStats.ComputePSOsWritten, OutStats.ComputePSOsFailed,
+		FPlatformTime::ToSeconds64(FPlatformTime::Cycles64() - StartCycles));
+
+	return 0;
+}
+
+//--------------------------------------------------------------------------------------------------
+// CollectPipelineFileCachePSOs
+// Loads each .upipelinecache file from -BundleDir and processes all PSO types:
+//   - Graphics: resolves per-stage shader hashes from the cache map, builds a merged root
+//     signature across all active stages, and stores a graphics FSODBEntry.
+//   - Compute: if compute collection is active (-AllCompute or -ShkDir), warns if a compute PSO
+//     from the cache was not already collected. Otherwise collects it directly from the cache map.
+//   - RayTracing: ignored for now (no RT SODB support yet).
+// All PSO types go into the "default" group; per-type grouping for graphics is deferred.
+//--------------------------------------------------------------------------------------------------
+
+static int32 CollectPipelineFileCachePSOs(
+	const FGenerateSODBArgs& Args,
+	const FShaderCacheMap& ShaderCache,
+	TMap<FString, TUniquePtr<FSODBGroup>>& Groups,
+	FGenerateSODBStats& OutStats)
+{
+	TArray<FString> CacheFiles;
+	IFileManager::Get().FindFiles(CacheFiles, *Args.BundleDir, TEXT(".upipelinecache"));
+
+	if (CacheFiles.IsEmpty())
+	{
+		UE_LOG(LogASDTool, Warning, TEXT("No .upipelinecache files found in '%s'"), *Args.BundleDir);
+		return 0;
+	}
+
+	// All PSO types go into "default" for now (graphics grouping/splitting deferred)
+	const FString DefaultGroupName = TEXT("default");
+	TUniquePtr<FSODBGroup>& DefaultGroup = Groups.FindOrAdd(DefaultGroupName);
+	if (!DefaultGroup.IsValid())
+	{
+		DefaultGroup = MakeUnique<FSODBGroup>();
+		DefaultGroup->Name = DefaultGroupName;
+	}
+	FSODBGroup* Group = DefaultGroup.Get();
+
+	// Whether compute PSOs are already being collected via CollectComputePSOs.
+	// ShkDir only sets up group routing -- AllCompute is the actual compute collection signal.
+	const bool bComputeAlreadyCollected = Args.bAllCompute;
+
+	uint64 StartCycles = FPlatformTime::Cycles64();
+	int32 TotalGraphicsProcessed = 0;
+	int32 GraphicsWritten      = 0;
+	int32 GraphicsFailed       = 0;
+	int32 GraphicsDuplicates   = 0;
+	int32 ComputeVerified      = 0;
+	int32 ComputeMissing       = 0;
+	int32 ComputeDirectWritten = 0;
+	int32 ComputeDirectFailed  = 0;
+
+	for (const FString& File : CacheFiles)
+	{
+		const FString CachePath = Args.BundleDir / File;
+
+		TSet<FPipelineCacheFileFormatPSO> PSOs;
+		if (!FPipelineFileCacheManager::LoadPipelineFileCacheInto(CachePath, PSOs))
+		{
+			UE_LOG(LogASDTool, Warning, TEXT("Failed to load pipeline cache '%s'"), *File);
+			continue;
+		}
+
+		// Helper: look up a shader hash in the cache and return its DXBC + full blob.
+		// Logs at Verbose when a non-zero hash is missing -- indicates a cache/cook mismatch.
+		auto GetStageDXBC = [&](const FShaderHash& Hash, TArray<uint8>& OutDXBC, TArray<uint8>& OutFullBlob) -> bool
+		{
+			if (Hash.IsZero())
+			{
+				return false;
+			}
+			const FCachedShaderEntry* Entry = ShaderCache.Find(Hash);
+			if (!Entry || Entry->DXBCData.IsEmpty())
+			{
+				UE_LOG(LogASDTool, Verbose, TEXT("  Shader not found in cache: %s"), *Hash.ToString());
+				return false;
+			}
+			uint32 EntryDXBCMagic = 0;
+			if (Entry->DXBCData.Num() >= 4) { FMemory::Memcpy(&EntryDXBCMagic, Entry->DXBCData.GetData(), sizeof(uint32)); }
+			if (Entry->DXBCData.Num() >= 4 && EntryDXBCMagic != 0x43425844)
+			{
+				UE_LOG(LogASDTool, Warning, TEXT("  Shader has invalid DXBC magic: %s"), *Hash.ToString());
+				return false;
+			}
+			OutDXBC     = Entry->DXBCData;
+			OutFullBlob = Entry->FullBlob;
+			return true;
+		};
+
+		for (const FPipelineCacheFileFormatPSO& PSO : PSOs)
+		{
+			switch (PSO.Type)
+			{
+				case FPipelineCacheFileFormatPSO::DescriptorType::Graphics:
+				{
+					++TotalGraphicsProcessed;
+					const FPipelineCacheFileFormatPSO::GraphicsDescriptor& GD = PSO.GraphicsDesc;
+					// Hash the full PSO descriptor using the engine-provided GetTypeHash, which covers all PSO state
+					// fields (shaders, blend, depth, raster, RT formats etc.) via a carefully constructed CRC32 chain.
+					// We intentionally rely on GetTypeHash rather than hashing the raw struct bytes -- the struct
+					// contains padding and may grow over time, and raw memory hashing would produce incorrect or
+					// unstable keys. The result is 32 bits, widened to FShaderHash (64-bit) for storage uniformity.
+					const uint32 PSOHash32 = GetTypeHash(PSO);
+					const FShaderHash PSOKeyHash = FXxHash64::HashBuffer(&PSOHash32, sizeof(PSOHash32));
+
+					TArray<uint8> DXBC_VS, DXBC_PS, DXBC_GS, DXBC_MS, DXBC_AS;
+					TArray<uint8> Blob_VS, Blob_PS, Blob_GS, Blob_MS, Blob_AS;
+
+					bool bHasVS = GetStageDXBC(GD.VertexShader,        DXBC_VS, Blob_VS);
+					bool bHasPS = GetStageDXBC(GD.FragmentShader,      DXBC_PS, Blob_PS);
+					bool bHasGS = GetStageDXBC(GD.GeometryShader,      DXBC_GS, Blob_GS);
+					bool bHasMS = GetStageDXBC(GD.MeshShader,          DXBC_MS, Blob_MS);
+					bool bHasAS = GetStageDXBC(GD.AmplificationShader, DXBC_AS, Blob_AS);
+
+					// Must have at least VS (traditional raster) or MS (mesh shader pipeline)
+					if (!bHasVS && !bHasMS)
+					{
+						++GraphicsFailed;
+						continue;
+					}
+
+					// Build merged root signature across all active stages
+					TRefCountPtr<ID3DBlob> RootSigBlob;
+					{
+						TArray<TConstArrayView<uint8>> StageBlobs;
+						if (!Blob_VS.IsEmpty()) StageBlobs.Add(TConstArrayView<uint8>(Blob_VS));
+						if (!Blob_PS.IsEmpty()) StageBlobs.Add(TConstArrayView<uint8>(Blob_PS));
+						if (!Blob_GS.IsEmpty()) StageBlobs.Add(TConstArrayView<uint8>(Blob_GS));
+						if (!Blob_MS.IsEmpty()) StageBlobs.Add(TConstArrayView<uint8>(Blob_MS));
+						if (!Blob_AS.IsEmpty()) StageBlobs.Add(TConstArrayView<uint8>(Blob_AS));
+
+						if (!BuildGraphicsShaderRootSignature(TConstArrayView<TConstArrayView<uint8>>(StageBlobs), D3D12_RESOURCE_BINDING_TIER_3, RootSigBlob))
+						{
+							UE_LOG(LogASDTool, Warning, TEXT("  Root sig build failed for graphics PSO (key=%s)"), *PSOKeyHash.ToString());
+							++GraphicsFailed;
+							continue;
+						}
+					}
+
+					if (Group->StoreGraphicsPSO(PSOKeyHash, DXBC_VS, DXBC_PS, DXBC_GS, DXBC_MS, DXBC_AS, RootSigBlob.GetReference()))
+					{
+						++GraphicsWritten;
+					}
+					else
+					{
+						++GraphicsDuplicates;
+					}
+
+					break;
+				}
+
+				case FPipelineCacheFileFormatPSO::DescriptorType::Compute:
+				{
+					const FShaderHash& ComputeHash = PSO.ComputeDesc.ComputeShader;
+					if (ComputeHash.IsZero())
+					{
+						break;
+					}
+
+					if (bComputeAlreadyCollected)
+					{
+						// Compute PSOs should already have been collected by CollectComputePSOs.
+						// Track verified vs missing so coverage gaps are visible in the summary.
+						if (Group->WrittenHashes.Contains(ComputeHash))
+						{
+							++ComputeVerified;
+						}
+						else
+						{
+							++ComputeMissing;
+							UE_LOG(LogASDTool, Verbose, TEXT("  Bundle compute PSO not in Step 4a: %s"), *ComputeHash.ToString());
+						}
+					}
+					else
+					{
+						// Compute collection not active -- collect this PSO directly from the cache map.
+						TArray<uint8> DXBC, FullBlob;
+						if (GetStageDXBC(ComputeHash, DXBC, FullBlob))
+						{
+							TRefCountPtr<ID3DBlob> RootSigBlob;
+							TConstArrayView<uint8> BlobView(FullBlob.GetData(), FullBlob.Num());
+							int32 NativeBytecodeOffset = 0;
+							if (BuildComputeShaderRootSignature(BlobView, D3D12_RESOURCE_BINDING_TIER_3, RootSigBlob, NativeBytecodeOffset))
+							{
+								if (Group->StoreComputeShader(ComputeHash, DXBC.GetData(), DXBC.Num(), RootSigBlob.GetReference()))
+								{
+									++ComputeDirectWritten;
+								}
+							}
+							else
+							{
+								UE_LOG(LogASDTool, Warning, TEXT("  Bundle compute root sig failed: %s"), *ComputeHash.ToString());
+								++ComputeDirectFailed;
+							}
+						}
+					}
+					break;
+				}
+
+				case FPipelineCacheFileFormatPSO::DescriptorType::RayTracing:
+				{
+					// Ray tracing PSOs are not yet supported -- ignored for now.
+					break;
+				}
+			}
+		}
+	}
+
+	// Update total stats
+	OutStats.GraphicsPSOsWritten += GraphicsWritten;
+	OutStats.GraphicsPSOsFailed += GraphicsFailed;
+	OutStats.ComputePSOsWritten += ComputeDirectWritten;
+	OutStats.ComputePSOsFailed += ComputeDirectFailed;
+
+	double BundleSeconds = FPlatformTime::ToSeconds64(FPlatformTime::Cycles64() - StartCycles);
+	int32 TotalBundleCompute = ComputeVerified + ComputeMissing + ComputeDirectWritten + ComputeDirectFailed;
+	UE_LOG(LogASDTool, Display, TEXT("  Bundle PSOs (%.2f s):"), BundleSeconds);
+	UE_LOG(LogASDTool, Display, TEXT("    Graphics:  %d processed  ->  %d written, %d duplicates, %d failed"),
+		TotalGraphicsProcessed, GraphicsWritten, GraphicsDuplicates, GraphicsFailed);
+	if (TotalBundleCompute > 0)
+	{
+		if (bComputeAlreadyCollected)
+		{
+			UE_LOG(LogASDTool, Display, TEXT("    Compute:   %d processed  ->  %d already collected, %d missing from ComputeCollection step"),
+				TotalBundleCompute, ComputeVerified, ComputeMissing);
+		}
+		else
+		{
+			UE_LOG(LogASDTool, Display, TEXT("    Compute:   %d processed  ->  %d written directly, %d failed"),
+				TotalBundleCompute, ComputeDirectWritten, ComputeDirectFailed);
+		}
+	}
+
+	return 0;
+}
+
+//--------------------------------------------------------------------------------------------------
+// WriteSODBFiles
+// Opens one ISODBFile writer per group and writes all FSODBEntry instances collected during
+// CollectComputePSOs and CollectPipelineFileCachePSOs. Uses the D3D12 API path by default;
+// -DirectWrite uses the SQLite path instead.
+//--------------------------------------------------------------------------------------------------
+
+static int32 WriteSODBFiles(
+	const FGenerateSODBArgs& Args,
+	TMap<FString, TUniquePtr<FSODBGroup>>& Groups)
+{
+	uint64 StartCycles = FPlatformTime::Cycles64();
+
+	// Build application desc for all writers
+	D3D12_APPLICATION_DESC AppDesc = {};
+	AppDesc.pExeFilename = *Args.ExeFilename;
+	AppDesc.pEngineName  = *Args.EngineName;
+	AppDesc.pName        = *Args.AppName;
+	AppDesc.EngineVersion.VersionParts[0] = FEngineVersion::Current().GetMajor();
+	AppDesc.EngineVersion.VersionParts[1] = FEngineVersion::Current().GetMinor();
+	AppDesc.EngineVersion.VersionParts[2] = FEngineVersion::Current().GetPatch();
+	AppDesc.Version.VersionParts[0] = 1;
+	AppDesc.Version.VersionParts[1] = 0;
+	AppDesc.Version.VersionParts[2] = FEngineVersion::Current().GetChangelist() >> 16;
+	AppDesc.Version.VersionParts[3] = FEngineVersion::Current().GetChangelist() & 0xffff;
+
+	// Initialize factory if needed for D3D12 path
+	FSODBFactory SODBFactory;
+	if (!Args.bDirectWrite)
+	{
+		if (!SODBFactory.Initialize())
+		{
+			return 3;
+		}
+	}
+
+	for (auto& Pair : Groups)
+	{
+		FSODBGroup& Group = *Pair.Value;
+		FString SODBPath = Args.OutputDir / (Group.Name + TEXT(".sodb"));
+
+		TUniquePtr<ISODBFile> Writer;
+		if (Args.bDirectWrite)
+		{
+			Writer = MakeUnique<FDirectSODBFile>();
+		}
+		else
+		{
+			auto D3D12File = MakeUnique<FD3D12SODBFile>();
+			D3D12File->SetFactory(&SODBFactory);
+			Writer = MoveTemp(D3D12File);
+		}
+
+		if (!Writer->Open(SODBPath, ESODBOpenMode::Write))
+		{
+			UE_LOG(LogASDTool, Error, TEXT("Failed to create SODB for group '%s': %s"), *Group.Name, *SODBPath);
+			return 3;
+		}
+
+		Writer->SetApplicationDesc(AppDesc);
+
+		if (!Writer->WriteEntries(Group.CollectedEntries))
+		{
+			UE_LOG(LogASDTool, Warning, TEXT("Some entries failed to write for group '%s'."), *Group.Name);
+		}
+
+		Writer->Close();
+	}
+
+	UE_LOG(LogASDTool, Display, TEXT("  Write time: %.2f s"), FPlatformTime::ToSeconds64(FPlatformTime::Cycles64() - StartCycles));
+	return 0;
+}
+
+//--------------------------------------------------------------------------------------------------
+// PrintStats
+// Logs per-group results (wanted vs written, any missing hashes) and the final generation stats.
+//--------------------------------------------------------------------------------------------------
+
+static void PrintStats(
+	const FGenerateSODBArgs& Args,
+	const TMap<FString, TUniquePtr<FSODBGroup>>& Groups,
+	const FGenerateSODBStats& Stats)
+{
+	UE_LOG(LogASDTool, Display, TEXT(""));
+	UE_LOG(LogASDTool, Display, TEXT("=== Per-Group Results ==="));
+
+	for (auto& Pair : Groups)
+	{
+		const FSODBGroup& Group = *Pair.Value;
+		FString SODBPath = Args.OutputDir / (Group.Name + TEXT(".sodb"));
+		int32 Wanted  = Group.WantedHashes.Num();
+		int32 Written = Group.WrittenHashes.Num();
+		double DXILSizeMB = (double)Group.TotalDXILBytes / (1024.0 * 1024.0);
+
+		if (Wanted > 0)
+		{
+			// SHK-based group: show wanted vs written vs missing
+			int32 Missing = Wanted - Written;
+			UE_LOG(LogASDTool, Display, TEXT("  %-40s  wanted: %d  written: %d  missing: %d  DXIL: %.1f MB  -> %s"),
+				*Group.Name, Wanted, Written, Missing, DXILSizeMB, *SODBPath);
+			if (Missing > 0)
+			{
+				int32 Logged = 0;
+				for (const FShaderHash& Hash : Group.WantedHashes)
+				{
+					if (!Group.WrittenHashes.Contains(Hash))
+					{
+						UE_LOG(LogASDTool, Warning, TEXT("    Missing: %s"), *Hash.ToString());
+						if (++Logged >= 10)
+						{
+							UE_LOG(LogASDTool, Warning, TEXT("    ... and %d more"), Missing - Logged);
+							break;
+						}
+					}
+				}
+			}
+		}
+		else
+		{
+			// Default group (no SHK): just show written count
+			UE_LOG(LogASDTool, Display, TEXT("  %-40s  written: %d  DXIL: %.1f MB  -> %s"),
+				*Group.Name, Written, DXILSizeMB, *SODBPath);
+		}
+	}
+
+	UE_LOG(LogASDTool, Display, TEXT(""));
+	UE_LOG(LogASDTool, Display, TEXT("=== SODB Generation Summary ==="));
+	UE_LOG(LogASDTool, Display, TEXT("  Compute PSOs:           %d written  (%d failed)"),   Stats.ComputePSOsWritten,    Stats.ComputePSOsFailed);
+	UE_LOG(LogASDTool, Display, TEXT("  Graphics PSOs:          %d written  (%d failed)"),  Stats.GraphicsPSOsWritten, Stats.GraphicsPSOsFailed);
+	UE_LOG(LogASDTool, Display, TEXT("  Total time:             %.1f s"), Stats.ElapsedSeconds);
+}
+
+//--------------------------------------------------------------------------------------------------
+// GenerateSODBs -- entry point
+//
+// Steps:
+//   Step 0:  Parse and validate command-line arguments
+//   Step 1:  Prepare output directory (clean stale .sodb files if -Clean, then create)
+//   Step 2:  Set up SODB groups (SHK-based named groups, or single "default" group)
+//   Step 3:  Pre-load shader bytecode cache (all SM6 archives -> FShaderCacheMap)
+//   Step 4a: Collect compute PSOs from shader cache (if -ShkDir or -AllCompute)
+//   Step 4b: Collect PSOs from pipeline cache files (if -BundleDir: graphics + compute)
+//   Step 5:  Write SODB files (one per group, batch write of all collected entries)
+//   Step 6:  Print per-group results and generation summary
+//--------------------------------------------------------------------------------------------------
+
+int32 GenerateSODBs()
+{
+	const TCHAR* CmdLine = FCommandLine::Get();
+
+	if (FParse::Param(CmdLine, TEXT("h")) || FParse::Param(CmdLine, TEXT("help")))
+	{
+		PrintUsage();
+		return 0;
+	}
+
+	uint64 TotalStartCycles = FPlatformTime::Cycles64();
+
+	// -------------------------------------------------------------------------
+	// Step 0: Parse and validate command-line arguments
+	// -------------------------------------------------------------------------
+
+	FGenerateSODBArgs Args;
+	if (!ParseGenerateSODBArgs(Args))
+	{
+		return 1;
+	}
+
+	// -------------------------------------------------------------------------
+	// Step 1: Prepare output directory
+	// -------------------------------------------------------------------------
+
+	int32 Result = PrepareOutputDirectory(Args);
+	if (Result != 0)
+	{
+		return Result;
+	}
+
+	// -------------------------------------------------------------------------
+	// Step 2: Set up SODB groups
+	// -------------------------------------------------------------------------
+
+	TMap<FString, TUniquePtr<FSODBGroup>> Groups;
+	Result = SetupGroups(Args, Groups);
+	if (Result != 0)
+	{
+		return Result;
+	}
+
+	// -------------------------------------------------------------------------
+	// Step 3: Pre-load shader bytecode cache
+	// -------------------------------------------------------------------------
+
+	FShaderCacheMap ShaderCache;
+	PreloadShaderCache(Args, ShaderCache);
+
+	// -------------------------------------------------------------------------
+	// Step 4a: Collect compute PSOs from shader cache
+	// -------------------------------------------------------------------------
+
+	FGenerateSODBStats Stats;
+
+	if (Args.bAllCompute)
+	{
+		Result = CollectComputePSOs(Args, ShaderCache, Groups, Stats);
+		if (Result != 0)
+		{
+			return Result;
+		}
+	}
+
+	// -------------------------------------------------------------------------
+	// Step 4b: Collect PSOs from pipeline cache files
+	// -------------------------------------------------------------------------
+
+	if (!Args.BundleDir.IsEmpty())
+	{
+		Result = CollectPipelineFileCachePSOs(Args, ShaderCache, Groups, Stats);
+		if (Result != 0)
+		{
+			return Result;
+		}
+	}
+
+	// -------------------------------------------------------------------------
+	// Step 5: Write SODB files
+	// -------------------------------------------------------------------------
+
+	Result = WriteSODBFiles(Args, Groups);
+	if (Result != 0)
+	{
+		return Result;
+	}
+
+	// -------------------------------------------------------------------------
+	// Step 6: Print summary
+	// -------------------------------------------------------------------------
+
+	Stats.ElapsedSeconds = FPlatformTime::ToSeconds64(FPlatformTime::Cycles64() - TotalStartCycles);
+	PrintStats(Args, Groups, Stats);
+
+	return 0;
+}
+
+}; // namespace ASDTool
+
diff --git a/Engine/Source/Programs/Windows/ASDTool/Private/RegisterPSDBs.cpp b/Engine/Source/Programs/Windows/ASDTool/Private/RegisterPSDBs.cpp
new file mode 100644
--- /dev/null
+++ b/Engine/Source/Programs/Windows/ASDTool/Private/RegisterPSDBs.cpp
@@ -0,0 +1,781 @@
+// Copyright Epic Games, Inc. All Rights Reserved.
+
+#include "ASDToolCommands.h"
+#include "ASDToolUtil.h"
+
+#include "Misc/CommandLine.h"
+#include "Misc/Parse.h"
+#include "Misc/Paths.h"
+#include "Misc/FileHelper.h"
+#include "HAL/FileManager.h"
+#include "HAL/PlatformTime.h"
+#include "HAL/PlatformAtomics.h"
+#include "Misc/EngineVersion.h"
+
+#include "Compression/OodleDataCompression.h"
+
+
+// COM and D3D12 headers for shader cache registration
+#include "Windows/AllowWindowsPlatformTypes.h"
+THIRD_PARTY_INCLUDES_START
+#include <initguid.h>
+#include <combaseapi.h>
+#include <d3d12.h>
+#include <dxgi.h>
+#include <d3dshadercacheregistration.h>
+THIRD_PARTY_INCLUDES_END
+#include "Windows/HideWindowsPlatformTypes.h"
+
+
+
+namespace ASDTool
+{
+
+// ------------------------------------------------------------------
+// COM helper: ID3DShaderCacheInstallerClient implementation
+// ------------------------------------------------------------------
+
+class FShaderCacheInstallerClient : public ID3DShaderCacheInstallerClient
+{
+public:
+	FShaderCacheInstallerClient(const FString& InInstallerName)
+		: InstallerName(InInstallerName)
+		, RefCount(1)
+	{
+	}
+
+	virtual ~FShaderCacheInstallerClient() = default;
+
+	// IUnknown
+	HRESULT STDMETHODCALLTYPE QueryInterface(REFIID riid, void** ppvObject)
+	{
+		if (riid == __uuidof(IUnknown) || riid == IID_ID3DShaderCacheInstallerClient)
+		{
+			*ppvObject = static_cast<ID3DShaderCacheInstallerClient*>(this);
+			AddRef();
+			return S_OK;
+		}
+		*ppvObject = nullptr;
+		return E_NOINTERFACE;
+	}
+
+	ULONG STDMETHODCALLTYPE AddRef() { return (ULONG)FPlatformAtomics::InterlockedIncrement(&RefCount); }
+	ULONG STDMETHODCALLTYPE Release()
+	{
+		ULONG Count = (ULONG)FPlatformAtomics::InterlockedDecrement(&RefCount);
+		if (Count == 0) { delete this; }
+		return Count;
+	}
+
+	// ID3DShaderCacheInstallerClient
+	HRESULT STDMETHODCALLTYPE GetInstallerName(SIZE_T* pNameLength, wchar_t* pName)
+	{
+		SIZE_T RequiredLength = InstallerName.Len() + 1;
+		if (!pName)
+		{
+			*pNameLength = RequiredLength;
+			return S_OK;
+		}
+		if (*pNameLength < RequiredLength)
+		{
+			*pNameLength = RequiredLength;
+			return S_FALSE;
+		}
+		FCString::Strncpy(pName, *InstallerName, *pNameLength);
+		*pNameLength = RequiredLength;
+		return S_OK;
+	}
+
+	D3D_SHADER_CACHE_APP_REGISTRATION_SCOPE STDMETHODCALLTYPE GetInstallerScope()
+	{
+		return D3D_SHADER_CACHE_APP_REGISTRATION_SCOPE_USER;
+	}
+
+	HRESULT STDMETHODCALLTYPE HandleDriverUpdate(ID3DShaderCacheInstaller* pInstaller)
+	{
+		UE_LOG(LogASDTool, Display, TEXT("Driver update notification received"));
+		return S_OK;
+	}
+
+private:
+	FString InstallerName;
+	volatile int32 RefCount;
+};
+
+// ------------------------------------------------------------------
+// Shader Cache Registration wrapper
+// ------------------------------------------------------------------
+
+class FShaderCacheRegistration
+{
+public:
+	~FShaderCacheRegistration() { Shutdown(); }
+
+	bool Initialize(const FString& InstallerName)
+	{
+		HRESULT hr = CoInitializeEx(nullptr, COINIT_MULTITHREADED);
+		if (hr == RPC_E_CHANGED_MODE)
+		{
+			// COM is already initialized on this thread with a different apartment model.
+			// This is non-fatal -- do not treat as failure, and do not call CoUninitialize on exit.
+			UE_LOG(LogASDTool, Verbose, TEXT("COM already initialized with different apartment model (RPC_E_CHANGED_MODE) -- continuing"));
+		}
+		else if (FAILED(hr))
+		{
+			UE_LOG(LogASDTool, Error, TEXT("Failed to initialize COM: 0x%08X"), hr);
+			return false;
+		}
+		else
+		{
+			bCOMInitialized = true;
+		}
+
+		hr = D3D12GetInterface(CLSID_D3DShaderCacheInstallerFactory, IID_PPV_ARGS(&Factory));
+		if (FAILED(hr) || !Factory)
+		{
+			UE_LOG(LogASDTool, Error, TEXT("Failed to get D3DShaderCacheInstallerFactory: 0x%08X"), hr);
+			Shutdown();  // COM was initialized -- must uninitialize before returning
+			return false;
+		}
+
+		FShaderCacheInstallerClient* Client = new FShaderCacheInstallerClient(InstallerName);
+		hr = Factory->CreateInstaller(Client, IID_PPV_ARGS(&Installer));
+		Client->Release();
+
+		if (FAILED(hr) || !Installer)
+		{
+			UE_LOG(LogASDTool, Error, TEXT("Failed to create shader cache installer: 0x%08X"), hr);
+			Shutdown();  // COM was initialized -- must uninitialize before returning
+			return false;
+		}
+
+		UE_LOG(LogASDTool, Display, TEXT("Shader cache installer initialized"));
+
+		// Create explorer (needed for Unregister to look up apps by exe path)
+		IDXGIFactory* DXGIFactory = nullptr;
+		IDXGIAdapter* Adapter = nullptr;
+		hr = CreateDXGIFactory(__uuidof(IDXGIFactory), (void**)&DXGIFactory);
+		if (SUCCEEDED(hr))
+		{
+			DXGIFactory->EnumAdapters(0, &Adapter);
+			DXGIFactory->Release();
+		}
+
+		hr = Factory->CreateExplorer(Adapter, IID_PPV_ARGS(&Explorer));
+		if (Adapter) { Adapter->Release(); }
+
+		if (FAILED(hr) || !Explorer)
+		{
+			UE_LOG(LogASDTool, Warning, TEXT("Failed to create shader cache explorer: 0x%08X (Unregister will not work)"), hr);
+		}
+
+		return true;
+	}
+
+	bool RegisterApplication(
+		const FString& ExePath,
+		const FString& ExeFilename,
+		const FString& AppName,
+		const FString& EngineName,
+		uint64 AppVersion,
+		uint64 EngineVersion)
+	{
+		if (!Installer) { return false; }
+
+		D3D_SHADER_CACHE_APPLICATION_DESC AppDesc = {};
+		AppDesc.pExeFilename = *ExeFilename;
+		AppDesc.pName = *AppName;
+		AppDesc.Version.Version = AppVersion;
+		AppDesc.pEngineName = *EngineName;
+		AppDesc.EngineVersion.Version = EngineVersion;
+
+		HRESULT hr = Installer->RegisterApplication(*ExePath, &AppDesc, IID_PPV_ARGS(&Application));
+		if (FAILED(hr) || !Application)
+		{
+			UE_LOG(LogASDTool, Error, TEXT("Failed to register application: 0x%08X"), hr);
+			return false;
+		}
+
+		UE_LOG(LogASDTool, Display, TEXT("Application registered: %s"), *AppName);
+		UE_LOG(LogASDTool, Display, TEXT("  Exe: %s"), *ExeFilename);
+		UE_LOG(LogASDTool, Display, TEXT("  Path: %s"), *ExePath);
+		return true;
+	}
+
+	bool RegisterComponent(
+		const FString& ComponentName,
+		const FString& SODBPath,
+		const TArray<TPair<FString, FString>>& PSDBEntries) // pairs of (AdapterFamily, PSDBPath)
+	{
+		if (!Application) { return false; }
+
+		TArray<D3D_SHADER_CACHE_PSDB_PROPERTIES> Props;
+		for (const auto& Entry : PSDBEntries)
+		{
+			D3D_SHADER_CACHE_PSDB_PROPERTIES Prop = {};
+			Prop.pAdapterFamily = *Entry.Key;
+			Prop.pPsdbPath = *Entry.Value;
+			Props.Add(Prop);
+		}
+
+		ID3DShaderCacheComponent* Component = nullptr;
+		HRESULT hr = Application->RegisterComponent(
+			*ComponentName,
+			*SODBPath,
+			Props.Num(),
+			Props.GetData(),
+			IID_PPV_ARGS(&Component));
+
+		if (FAILED(hr))
+		{
+			UE_LOG(LogASDTool, Error, TEXT("Failed to register component '%s': 0x%08X"), *ComponentName, hr);
+			return false;
+		}
+
+		UE_LOG(LogASDTool, Display, TEXT("Component registered: %s"), *ComponentName);
+		UE_LOG(LogASDTool, Display, TEXT("  SODB: %s"), *SODBPath);
+		for (const auto& Entry : PSDBEntries)
+		{
+			UE_LOG(LogASDTool, Display, TEXT("  PSDB [%s]: %s"), *Entry.Key, *Entry.Value);
+		}
+
+		if (Component) { Component->Release(); }
+		return true;
+	}
+
+	bool QueryPrecompileTargets(
+		const FString& ExeFilename,
+		const FString& AppName,
+		const FString& EngineName,
+		uint64 AppVersion,
+		uint64 EngineVersion)
+	{
+		if (!Installer) { return false; }
+
+		uint32 MaxCount = Installer->GetMaxPrecompileTargetCount();
+		if (MaxCount == 0)
+		{
+			UE_LOG(LogASDTool, Display, TEXT("No precompile targets available"));
+			return true;
+		}
+
+		D3D_SHADER_CACHE_APPLICATION_DESC AppDesc = {};
+		AppDesc.pExeFilename = *ExeFilename;
+		AppDesc.pName = *AppName;
+		AppDesc.Version.Version = AppVersion;
+		AppDesc.pEngineName = *EngineName;
+		AppDesc.EngineVersion.Version = EngineVersion;
+
+		TArray<D3D_SHADER_CACHE_COMPILER_PROPERTIES> Targets;
+		Targets.SetNum(MaxCount);
+		uint32 TargetCount = MaxCount;
+
+		HRESULT hr = Installer->GetPrecompileTargets(
+			&AppDesc, &TargetCount, Targets.GetData(),
+			D3D_SHADER_CACHE_TARGET_FLAG_NONE);
+
+		if (FAILED(hr))
+		{
+			UE_LOG(LogASDTool, Error, TEXT("Failed to get precompile targets: 0x%08X"), hr);
+			return false;
+		}
+
+		UE_LOG(LogASDTool, Display, TEXT(""));
+		UE_LOG(LogASDTool, Display, TEXT("=== Available Precompile Targets ==="));
+		for (uint32 i = 0; i < TargetCount; i++)
+		{
+			UE_LOG(LogASDTool, Display, TEXT("  Target %d:"), i);
+			UE_LOG(LogASDTool, Display, TEXT("    Adapter Family: %s"), Targets[i].szAdapterFamily);
+			UE_LOG(LogASDTool, Display, TEXT("    ABI Support: %llu - %llu"),
+				Targets[i].MinimumABISupportVersion, Targets[i].MaximumABISupportVersion);
+			UE_LOG(LogASDTool, Display, TEXT("    Compiler Version: 0x%016llX"), Targets[i].CompilerVersion.Version);
+			UE_LOG(LogASDTool, Display, TEXT("    App Profile Version: 0x%016llX"), Targets[i].ApplicationProfileVersion.Version);
+		}
+
+		return true;
+	}
+
+	bool ListRegistrations()
+	{
+		if (!Installer) { return false; }
+
+		uint32 AppCount = Installer->GetApplicationCount();
+		UE_LOG(LogASDTool, Display, TEXT(""));
+		UE_LOG(LogASDTool, Display, TEXT("=== Registered Applications: %d ==="), AppCount);
+
+		for (uint32 i = 0; i < AppCount; i++)
+		{
+			ID3DShaderCacheApplication* App = nullptr;
+			HRESULT hr = Installer->GetApplication(i, IID_PPV_ARGS(&App));
+			if (FAILED(hr) || !App) continue;
+
+			const wchar_t* ExePath = nullptr;
+			App->GetExePath(&ExePath);
+
+			D3D_SHADER_CACHE_APPLICATION_DESC Desc = {};
+			App->GetDesc(&Desc);
+
+			UE_LOG(LogASDTool, Display, TEXT(""));
+			UE_LOG(LogASDTool, Display, TEXT("  [%d] %s"), i, Desc.pName ? Desc.pName : TEXT("(null)"));
+			UE_LOG(LogASDTool, Display, TEXT("      Exe: %s"), Desc.pExeFilename ? Desc.pExeFilename : TEXT("(null)"));
+			UE_LOG(LogASDTool, Display, TEXT("      Path: %s"), ExePath ? ExePath : TEXT("(null)"));
+			UE_LOG(LogASDTool, Display, TEXT("      Engine: %s"), Desc.pEngineName ? Desc.pEngineName : TEXT("(null)"));
+
+			uint32 CompCount = App->GetComponentCount();
+			for (uint32 j = 0; j < CompCount; j++)
+			{
+				ID3DShaderCacheComponent* Comp = nullptr;
+				hr = App->GetComponent(j, IID_PPV_ARGS(&Comp));
+				if (FAILED(hr) || !Comp) continue;
+
+				const wchar_t* CompName = nullptr;
+				Comp->GetComponentName(&CompName);
+
+				const wchar_t* SODBPath = nullptr;
+				Comp->GetStateObjectDatabasePath(&SODBPath);
+
+				UE_LOG(LogASDTool, Display, TEXT("      Component [%d]: %s"), j, CompName ? CompName : TEXT("(null)"));
+				UE_LOG(LogASDTool, Display, TEXT("        SODB: %s"), SODBPath ? SODBPath : TEXT("(null)"));
+
+				uint32 PSDBCount = Comp->GetPrecompiledShaderDatabaseCount();
+				if (PSDBCount > 0)
+				{
+					TArray<D3D_SHADER_CACHE_PSDB_PROPERTIES> PSDBProps;
+					PSDBProps.SetNum(PSDBCount);
+					if (SUCCEEDED(Comp->GetPrecompiledShaderDatabases(PSDBCount, PSDBProps.GetData())))
+					{
+						for (uint32 k = 0; k < PSDBCount; k++)
+						{
+							UE_LOG(LogASDTool, Display, TEXT("        PSDB [%s]: %s"),
+								PSDBProps[k].pAdapterFamily ? PSDBProps[k].pAdapterFamily : TEXT("(null)"),
+								PSDBProps[k].pPsdbPath ? PSDBProps[k].pPsdbPath : TEXT("(null)"));
+						}
+					}
+				}
+
+				Comp->Release();
+			}
+
+			App->Release();
+		}
+
+		return true;
+	}
+
+	bool ClearAll()
+	{
+		if (!Installer) { return false; }
+
+		HRESULT hr = Installer->ClearAllState();
+		if (FAILED(hr))
+		{
+			UE_LOG(LogASDTool, Error, TEXT("Failed to clear all registrations: 0x%08X"), hr);
+			return false;
+		}
+
+		UE_LOG(LogASDTool, Display, TEXT("All shader cache registrations cleared"));
+		return true;
+	}
+
+	bool Unregister(const FString& ExePath)
+	{
+		if (!Explorer)
+		{
+			UE_LOG(LogASDTool, Error, TEXT("Shader cache explorer not available - cannot unregister"));
+			return false;
+		}
+
+		ID3DShaderCacheApplication* AppToRemove = nullptr;
+		HRESULT hr = Explorer->GetApplicationFromExePath(*ExePath, IID_PPV_ARGS(&AppToRemove));
+
+		if (FAILED(hr) || !AppToRemove)
+		{
+			UE_LOG(LogASDTool, Error, TEXT("Application not found for path: %s (0x%08X)"), *ExePath, hr);
+			return false;
+		}
+
+		hr = Installer->RemoveApplication(AppToRemove);
+		AppToRemove->Release();
+
+		if (FAILED(hr))
+		{
+			UE_LOG(LogASDTool, Error, TEXT("Failed to remove application: 0x%08X"), hr);
+			return false;
+		}
+
+		UE_LOG(LogASDTool, Display, TEXT("Application unregistered: %s"), *ExePath);
+		return true;
+	}
+
+	/** Get the adapter family for the locally installed GPU. */
+	bool GetLocalAdapterFamily(
+		const FString& ExeFilename,
+		const FString& AppName,
+		const FString& EngineName,
+		uint64 AppVersion,
+		uint64 EngineVersion,
+		FString& OutFamily)
+	{
+		if (!Installer) { return false; }
+
+		uint32 MaxCount = Installer->GetMaxPrecompileTargetCount();
+		if (MaxCount == 0) { return false; }
+
+		D3D_SHADER_CACHE_APPLICATION_DESC AppDesc = {};
+		AppDesc.pExeFilename = *ExeFilename;
+		AppDesc.pName = *AppName;
+		AppDesc.Version.Version = AppVersion;
+		AppDesc.pEngineName = *EngineName;
+		AppDesc.EngineVersion.Version = EngineVersion;
+
+		TArray<D3D_SHADER_CACHE_COMPILER_PROPERTIES> Targets;
+		Targets.SetNum(MaxCount);
+		uint32 TargetCount = MaxCount;
+
+		HRESULT hr = Installer->GetPrecompileTargets(
+			&AppDesc, &TargetCount, Targets.GetData(),
+			D3D_SHADER_CACHE_TARGET_FLAG_NONE);
+
+		if (FAILED(hr) || TargetCount == 0) { return false; }
+
+		// Use the first target's adapter family. On machines with multiple adapters (e.g. iGPU + dGPU)
+		// the first target is typically the primary display adapter, which is the correct choice
+		// for registration. If needed, the caller can override via -AdapterFamily.
+		OutFamily = Targets[0].szAdapterFamily;
+		return true;
+	}
+
+	void Shutdown()
+	{
+		if (Application) { Application->Release(); Application = nullptr; }
+		if (Explorer)    { Explorer->Release();    Explorer = nullptr; }
+		if (Installer)   { Installer->Release();   Installer = nullptr; }
+		if (Factory)     { Factory->Release();     Factory = nullptr; }
+		if (bCOMInitialized) { CoUninitialize(); bCOMInitialized = false; }
+	}
+
+private:
+	ID3DShaderCacheInstallerFactory* Factory = nullptr;
+	ID3DShaderCacheInstaller* Installer = nullptr;
+	ID3DShaderCacheApplication* Application = nullptr;
+	ID3DShaderCacheExplorer* Explorer = nullptr;
+	bool bCOMInitialized = false;
+};
+
+// ------------------------------------------------------------------
+// Command-line arguments
+// ------------------------------------------------------------------
+
+struct FRegisterPSDBsArgs
+{
+	FString PSDBDir;
+	FString SODBPath;
+	FString OutputDir;
+	FString AppName;
+	FString ExeFilename;
+	FString ExePath;
+	FString EngineName;
+	FString ComponentName;
+	FString AdapterFamilyOverride;
+	bool bListOnly = false;
+	bool bUnregister = false;
+	bool bClearAll = false;
+	bool bQueryTargets = false;
+};
+
+// ------------------------------------------------------------------
+// Usage
+// ------------------------------------------------------------------
+
+static void PrintUsage()
+{
+	UE_LOG(LogASDTool, Display, TEXT("RegisterPSDBs: register Precompiled Shader Databases with the D3D12 runtime"));
+	UE_LOG(LogASDTool, Display, TEXT(""));
+	UE_LOG(LogASDTool, Display, TEXT("Required flags (for registration):"));
+	UE_LOG(LogASDTool, Display, TEXT("    -PSDBDir=<dir>            Directory with .psdb or .psdb.oodle files"));
+	UE_LOG(LogASDTool, Display, TEXT("    -SODBPath=<file>          Path to the SODB file"));
+	UE_LOG(LogASDTool, Display, TEXT("    -Name=<string>            Application name"));
+	UE_LOG(LogASDTool, Display, TEXT("    -ExeFilename=<string>     Application exe filename"));
+	UE_LOG(LogASDTool, Display, TEXT("    -ExePath=<path>           Full path to application exe"));
+	UE_LOG(LogASDTool, Display, TEXT("    -Engine=<string>          Engine name"));
+	UE_LOG(LogASDTool, Display, TEXT(""));
+	UE_LOG(LogASDTool, Display, TEXT("Optional flags:"));
+	UE_LOG(LogASDTool, Display, TEXT("    -ComponentName=<string>   Component name (default: \"default\")"));
+	UE_LOG(LogASDTool, Display, TEXT("    -OutputDir=<dir>          Where to write decompressed PSDBs (default: next to .oodle files)"));
+	UE_LOG(LogASDTool, Display, TEXT("    -AdapterFamily=<string>   Register only for this adapter family (default: auto-detect local GPU)"));
+	UE_LOG(LogASDTool, Display, TEXT("    -ListOnly                 List existing registrations"));
+	UE_LOG(LogASDTool, Display, TEXT("    -Unregister               Remove registration for this app"));
+	UE_LOG(LogASDTool, Display, TEXT("    -ClearAll                 Clear ALL registrations (WARNING)"));
+	UE_LOG(LogASDTool, Display, TEXT("    -QueryTargets             Query available precompile targets"));
+	UE_LOG(LogASDTool, Display, TEXT("    -Help                     Print this help message"));
+}
+
+// Always returns true -- required field validation is deferred to ValidateRegistrationArgs()
+// so that utility commands (list, clear, unregister, query) can run without providing
+// registration-only args like -PSDBDir or -SODBPath.
+static bool ParseRegisterPSDBsArgs(FRegisterPSDBsArgs& OutArgs)
+{
+	const TCHAR* CmdLine = FCommandLine::Get();
+
+	OutArgs.bListOnly     = FParse::Param(CmdLine, TEXT("ListOnly"));
+	OutArgs.bUnregister   = FParse::Param(CmdLine, TEXT("Unregister"));
+	OutArgs.bClearAll     = FParse::Param(CmdLine, TEXT("ClearAll"));
+	OutArgs.bQueryTargets = FParse::Param(CmdLine, TEXT("QueryTargets"));
+
+	FParse::Value(CmdLine, TEXT("-PSDBDir="), OutArgs.PSDBDir);
+	FParse::Value(CmdLine, TEXT("-SODBPath="), OutArgs.SODBPath);
+	FParse::Value(CmdLine, TEXT("-OutputDir="), OutArgs.OutputDir);
+	FParse::Value(CmdLine, TEXT("-Name="), OutArgs.AppName);
+	FParse::Value(CmdLine, TEXT("-ExeFilename="), OutArgs.ExeFilename);
+	FParse::Value(CmdLine, TEXT("-ExePath="), OutArgs.ExePath);
+	FParse::Value(CmdLine, TEXT("-Engine="), OutArgs.EngineName);
+	FParse::Value(CmdLine, TEXT("-ComponentName="), OutArgs.ComponentName);
+	FParse::Value(CmdLine, TEXT("-AdapterFamily="), OutArgs.AdapterFamilyOverride);
+
+	if (OutArgs.ComponentName.IsEmpty())
+	{
+		OutArgs.ComponentName = TEXT("default");
+	}
+
+	return true;
+}
+
+/** Validate that all required arguments for registration are present. */
+static bool ValidateRegistrationArgs(const FRegisterPSDBsArgs& Args)
+{
+	if (Args.PSDBDir.IsEmpty() || Args.SODBPath.IsEmpty() || Args.AppName.IsEmpty() ||
+		Args.ExeFilename.IsEmpty() || Args.ExePath.IsEmpty() || Args.EngineName.IsEmpty())
+	{
+		PrintUsage();
+		if (Args.PSDBDir.IsEmpty())     UE_LOG(LogASDTool, Error, TEXT("Missing -PSDBDir"));
+		if (Args.SODBPath.IsEmpty())    UE_LOG(LogASDTool, Error, TEXT("Missing -SODBPath"));
+		if (Args.AppName.IsEmpty())     UE_LOG(LogASDTool, Error, TEXT("Missing -Name"));
+		if (Args.ExeFilename.IsEmpty()) UE_LOG(LogASDTool, Error, TEXT("Missing -ExeFilename"));
+		if (Args.ExePath.IsEmpty())     UE_LOG(LogASDTool, Error, TEXT("Missing -ExePath"));
+		if (Args.EngineName.IsEmpty())  UE_LOG(LogASDTool, Error, TEXT("Missing -Engine"));
+		return false;
+	}
+	return true;
+}
+
+/** Find PSDB files for the target adapter family. */
+static int32 FindPSDBFiles(
+	const FRegisterPSDBsArgs& Args,
+	const FString& AdapterFamily,
+	TArray<FString>& OutPSDBFiles)
+{
+	if (Args.PSDBDir.IsEmpty())
+	{
+		UE_LOG(LogASDTool, Error, TEXT("PSDBDir is empty -- cannot search for PSDB files"));
+		return 1;
+	}
+
+	// Sanitized family name for folder lookup (: replaced with _)
+	FString SafeFamilyName = AdapterFamily;
+	SafeFamilyName.ReplaceCharInline(TCHAR(':'), TCHAR('_'));
+
+	// Look for PSDBs in the directory structure: <PSDBDir>/<SafeFamilyName>/abi_<ver>/*.psdb[.oodle]
+	FString FamilyDir = Args.PSDBDir / SafeFamilyName;
+
+	if (FPaths::DirectoryExists(FamilyDir))
+	{
+		IFileManager::Get().FindFilesRecursive(OutPSDBFiles, *FamilyDir, TEXT("*.psdb"), true, false);
+		// Also find .psdb.oodle files
+		TArray<FString> OodleFiles;
+		IFileManager::Get().FindFilesRecursive(OodleFiles, *FamilyDir, TEXT("*.psdb.oodle"), true, false);
+		OutPSDBFiles.Append(OodleFiles);
+	}
+	else
+	{
+		// Fallback: search the entire PSDBDir for any PSDB files
+		IFileManager::Get().FindFilesRecursive(OutPSDBFiles, *Args.PSDBDir, TEXT("*.psdb"), true, false);
+		TArray<FString> OodleFiles;
+		IFileManager::Get().FindFilesRecursive(OodleFiles, *Args.PSDBDir, TEXT("*.psdb.oodle"), true, false);
+		OutPSDBFiles.Append(OodleFiles);
+	}
+
+	if (OutPSDBFiles.IsEmpty())
+	{
+		UE_LOG(LogASDTool, Error, TEXT("No PSDB files found for adapter family '%s' in '%s'"), *AdapterFamily, *Args.PSDBDir);
+		return 1;
+	}
+
+	UE_LOG(LogASDTool, Display, TEXT("Found %d PSDB file(s) for %s"), OutPSDBFiles.Num(), *AdapterFamily);
+
+	return 0;
+}
+
+/** Decompress .oodle files and build registration entries. */
+static int32 PrepareRegistrationEntries(
+	const FRegisterPSDBsArgs& Args,
+	const FString& AdapterFamily,
+	const TArray<FString>& PSDBFiles,
+	TArray<TPair<FString, FString>>& OutEntries)
+{
+	for (const FString& PSDBFile : PSDBFiles)
+	{
+		FString FinalPSDBPath = PSDBFile;
+
+		if (PSDBFile.EndsWith(TEXT(".oodle")))
+		{
+			// Determine decompressed output path.
+			// NOTE: When -OutputDir is specified, all PSDBs decompress into a flat directory
+			// using only the base filename. PSDBs with the same name from different adapter
+			// family subdirectories (e.g. AMD/default.psdb and NVIDIA/default.psdb) will
+			// collide. -OutputDir is intended for single-family use; for multi-family flows
+			// omit it and let each PSDB decompress alongside its source file.
+			FString DecompressedName = FPaths::GetBaseFilename(PSDBFile); // strips .oodle, leaves .psdb
+			FString DecompressedDir = Args.OutputDir.IsEmpty() ? FPaths::GetPath(PSDBFile) : Args.OutputDir;
+			FinalPSDBPath = DecompressedDir / DecompressedName;
+
+			if (!DecompressFile(PSDBFile, FinalPSDBPath, PSDB_COMPRESSED_MAGIC))
+			{
+				UE_LOG(LogASDTool, Error, TEXT("Failed to decompress: %s"), *PSDBFile);
+				return 1;
+			}
+		}
+
+		OutEntries.Add(TPair<FString, FString>(AdapterFamily, FPaths::ConvertRelativePathToFull(FinalPSDBPath)));
+	}
+
+	return 0;
+}
+
+// ------------------------------------------------------------------
+// Main
+// ------------------------------------------------------------------
+
+/** Handle utility-only commands (list, clear, unregister, query). Returns -1 if not a utility command. */
+static int32 HandleUtilityCommand(const FRegisterPSDBsArgs& Args, FShaderCacheRegistration& Registration, uint64 AppVersion, uint64 EngineVersion)
+{
+	if (Args.bClearAll)
+	{
+		return Registration.ClearAll() ? 0 : 1;
+	}
+
+	if (Args.bListOnly)
+	{
+		return Registration.ListRegistrations() ? 0 : 1;
+	}
+
+	if (Args.bUnregister)
+	{
+		if (Args.ExePath.IsEmpty())
+		{
+			UE_LOG(LogASDTool, Error, TEXT("Missing -ExePath for unregister"));
+			return 1;
+		}
+		return Registration.Unregister(Args.ExePath) ? 0 : 1;
+	}
+
+	if (Args.bQueryTargets)
+	{
+		if (Args.ExeFilename.IsEmpty() || Args.AppName.IsEmpty() || Args.EngineName.IsEmpty())
+		{
+			UE_LOG(LogASDTool, Error, TEXT("Missing -Name, -ExeFilename, or -Engine for query"));
+			return 1;
+		}
+		return Registration.QueryPrecompileTargets(Args.ExeFilename, Args.AppName, Args.EngineName, AppVersion, EngineVersion) ? 0 : 1;
+	}
+
+	return -1; // Not a utility command
+}
+
+int32 RegisterPSDBs()
+{
+	const TCHAR* CmdLine = FCommandLine::Get();
+
+	if (FParse::Param(CmdLine, TEXT("h")) || FParse::Param(CmdLine, TEXT("help")))
+	{
+		PrintUsage();
+		return 0;
+	}
+
+	// Parse command-line arguments
+	FRegisterPSDBsArgs Args;
+	if (!ParseRegisterPSDBsArgs(Args))
+	{
+		return 1;
+	}
+
+	// Version defaults (1.0 app, engine from UE version)
+	uint64 AppVersion = (1ULL << 48) | ((uint64)FEngineVersion::Current().GetChangelist() >> 16 << 16) | ((uint64)FEngineVersion::Current().GetChangelist() & 0xffff);
+	uint64 EngineVersion = ((uint64)FEngineVersion::Current().GetMajor() << 48) | ((uint64)FEngineVersion::Current().GetMinor() << 32) | ((uint64)FEngineVersion::Current().GetPatch() << 16);
+
+	// ---- Initialize ----
+
+	FShaderCacheRegistration Registration;
+	if (!Registration.Initialize(TEXT("ASDTool")))
+	{
+		return 1;
+	}
+
+	// ---- Handle utility commands ----
+
+	int32 UtilityResult = HandleUtilityCommand(Args, Registration, AppVersion, EngineVersion);
+	if (UtilityResult >= 0)
+	{
+		return UtilityResult;
+	}
+
+	// ---- Validate registration args ----
+
+	if (!ValidateRegistrationArgs(Args))
+	{
+		return 1;
+	}
+
+	// ---- Determine adapter family ----
+
+	FString AdapterFamily = Args.AdapterFamilyOverride;
+	if (AdapterFamily.IsEmpty())
+	{
+		UE_LOG(LogASDTool, Display, TEXT("Auto-detecting local adapter family..."));
+		if (!Registration.GetLocalAdapterFamily(Args.ExeFilename, Args.AppName, Args.EngineName, AppVersion, EngineVersion, AdapterFamily))
+		{
+			UE_LOG(LogASDTool, Error, TEXT("Failed to detect local adapter family. Use -AdapterFamily= to specify manually."));
+			return 1;
+		}
+	}
+	UE_LOG(LogASDTool, Display, TEXT("Target adapter family: %s"), *AdapterFamily);
+
+	// ---- Find PSDB files for this adapter family ----
+
+	TArray<FString> PSDBFiles;
+	int32 Result = FindPSDBFiles(Args, AdapterFamily, PSDBFiles);
+	if (Result != 0)
+	{
+		return Result;
+	}
+
+	// ---- Decompress .oodle files if needed ----
+
+	TArray<TPair<FString, FString>> RegisterEntries; // (AdapterFamily, PSDBPath)
+	Result = PrepareRegistrationEntries(Args, AdapterFamily, PSDBFiles, RegisterEntries);
+	if (Result != 0)
+	{
+		return Result;
+	}
+
+	// ---- Register application ----
+
+	if (!Registration.RegisterApplication(Args.ExePath, Args.ExeFilename, Args.AppName, Args.EngineName, AppVersion, EngineVersion))
+	{
+		return 1;
+	}
+
+	// ---- Register component with PSDB entries ----
+
+	FString FullSODBPath = FPaths::ConvertRelativePathToFull(Args.SODBPath);
+	if (!Registration.RegisterComponent(Args.ComponentName, FullSODBPath, RegisterEntries))
+	{
+		return 1;
+	}
+
+	UE_LOG(LogASDTool, Display, TEXT(""));
+	UE_LOG(LogASDTool, Display, TEXT("Registration completed successfully."));
+
+	return 0;
+}
+
+};
\ No newline at end of file
diff --git a/Engine/Source/Programs/Windows/ASDTool/Private/SODBFile.cpp b/Engine/Source/Programs/Windows/ASDTool/Private/SODBFile.cpp
new file mode 100644
--- /dev/null
+++ b/Engine/Source/Programs/Windows/ASDTool/Private/SODBFile.cpp
@@ -0,0 +1,1207 @@
+// Copyright Epic Games, Inc. All Rights Reserved.
+
+#include "SODBFile.h"
+#include "ASDToolUtil.h"
+
+// initguid.h must be included before any D3D12 headers to instantiate GUIDs
+#include "Windows/AllowWindowsPlatformTypes.h"
+THIRD_PARTY_INCLUDES_START
+#include <initguid.h>
+THIRD_PARTY_INCLUDES_END
+#include "Windows/HideWindowsPlatformTypes.h"
+
+#include "Misc/FileHelper.h"
+#include "Misc/Paths.h"
+#include "HAL/FileManager.h"
+#include "HAL/PlatformAtomics.h"
+#include "HAL/PlatformTime.h"
+#include "Misc/ScopeRWLock.h"
+#include "Misc/SecureHash.h"
+
+#include "IncludeSQLite.h"
+
+namespace ASDTool
+{
+
+//--------------------------------------------------------------------------------------------------
+// FSODBGroup
+//--------------------------------------------------------------------------------------------------
+
+bool FSODBGroup::Contains(const FShaderHash& Hash) const
+{
+	return WantedHashes.Contains(Hash);
+}
+
+bool FSODBGroup::DeduplicateAndMark(const FShaderHash& Hash)
+{
+	{
+		FReadScopeLock ReadLock(WrittenLock);
+		if (WrittenHashes.Contains(Hash))
+		{
+			return false;
+		}
+	}
+	{
+		FWriteScopeLock WriteLock(WrittenLock);
+		WrittenHashes.Add(Hash);
+	}
+	return true;
+}
+
+bool FSODBGroup::StoreComputeShader(const FShaderHash& Hash, const uint8* DXBCData, int64 DXBCSize, ID3DBlob* RootSignatureBlob)
+{
+	if (!DeduplicateAndMark(Hash))
+	{
+		return false;
+	}
+
+	FPlatformAtomics::InterlockedAdd(&TotalDXILBytes, DXBCSize);
+
+	// Build PSO stream desc (same layout as the old FD3D12SODBWriter)
+	struct FComputePipelineStream
+	{
+		struct
+		{
+			D3D12_PIPELINE_STATE_SUBOBJECT_TYPE Type;
+			D3D12_SHADER_BYTECODE CS;
+		} CSSubobject;
+	};
+
+	FComputePipelineStream Stream = {};
+	// NOTE: pShaderBytecode is set after the DXBC is copied into Entry.StageDXBC[SF_Compute]
+	// below, so the stream points at the entry-owned buffer rather than the caller's pointer.
+	Stream.CSSubobject.Type = D3D12_PIPELINE_STATE_SUBOBJECT_TYPE_CS;
+
+	FSODBEntry Entry;
+	uint8 KeyBytes[sizeof(uint64)];  // FShaderHash is FXxHash64 -- ToByteArray writes exactly 8 bytes
+	Hash.ToByteArray(KeyBytes);
+	Entry.Key.SetNumUninitialized(sizeof(KeyBytes));
+	FMemory::Memcpy(Entry.Key.GetData(), KeyBytes, sizeof(KeyBytes));
+	Entry.Version = 0;
+
+	// Store DXBC for the direct SQLite path (compute = SF_Compute index).
+	// Also point the stream's pShaderBytecode at this owned copy so StreamData doesn't
+	// dangle when the caller's DXBCData pointer goes out of scope.
+	Entry.StageDXBC[SF_Compute].SetNumUninitialized(DXBCSize);
+	FMemory::Memcpy(Entry.StageDXBC[SF_Compute].GetData(), DXBCData, DXBCSize);
+	Stream.CSSubobject.CS.pShaderBytecode = Entry.StageDXBC[SF_Compute].GetData();
+	Stream.CSSubobject.CS.BytecodeLength  = (SIZE_T)DXBCSize;
+
+	// Store root signature blob for the direct SQLite path
+	if (RootSignatureBlob)
+	{
+		Entry.RootSigBlob = RootSignatureBlob;
+	}
+
+	// Build StreamData for the D3D12 API path
+	if (RootSignatureBlob)
+	{
+		struct FRootSigSubobject
+		{
+			D3D12_PIPELINE_STATE_SUBOBJECT_TYPE Type;
+			D3D12_SERIALIZED_ROOT_SIGNATURE_DESC RootSig;
+		};
+
+		FRootSigSubobject RootSigSubobject = {};
+		RootSigSubobject.Type = D3D12_PIPELINE_STATE_SUBOBJECT_TYPE_SERIALIZED_ROOT_SIGNATURE;
+		RootSigSubobject.RootSig.pSerializedBlob = RootSignatureBlob->GetBufferPointer();
+		RootSigSubobject.RootSig.SerializedBlobSizeInBytes = RootSignatureBlob->GetBufferSize();
+
+		Entry.StreamData.SetNumUninitialized(sizeof(FRootSigSubobject) + sizeof(FComputePipelineStream));
+		FMemory::Memcpy(Entry.StreamData.GetData(), &RootSigSubobject, sizeof(FRootSigSubobject));
+		FMemory::Memcpy(Entry.StreamData.GetData() + sizeof(FRootSigSubobject), &Stream, sizeof(FComputePipelineStream));
+	}
+	else
+	{
+		Entry.StreamData.SetNumUninitialized(sizeof(FComputePipelineStream));
+		FMemory::Memcpy(Entry.StreamData.GetData(), &Stream, sizeof(FComputePipelineStream));
+	}
+
+	{
+		FScopeLock Lock(&EntriesLock);
+		CollectedEntries.Add(MoveTemp(Entry));
+	}
+
+	return true;
+}
+
+bool FSODBGroup::StoreGraphicsPSO(
+	const FShaderHash& Hash,
+	const TArray<uint8>& DXBC_VS, const TArray<uint8>& DXBC_PS,
+	const TArray<uint8>& DXBC_GS, const TArray<uint8>& DXBC_MS, const TArray<uint8>& DXBC_AS,
+	ID3DBlob* RootSignatureBlob)
+{
+	if (!DeduplicateAndMark(Hash))
+	{
+		return false;
+	}
+
+	FSODBEntry Entry;
+	Entry.bIsGraphics = true;
+	uint8 GraphicsKeyBytes[sizeof(uint64)];  // FShaderHash is FXxHash64 -- ToByteArray writes exactly 8 bytes
+	Hash.ToByteArray(GraphicsKeyBytes);
+	Entry.Key.SetNumUninitialized(sizeof(GraphicsKeyBytes));
+	FMemory::Memcpy(Entry.Key.GetData(), GraphicsKeyBytes, sizeof(GraphicsKeyBytes));
+	Entry.Version = 0;
+
+	// Store per-stage DXBC for the direct SQLite path
+	Entry.StageDXBC[SF_Vertex]        = DXBC_VS;
+	Entry.StageDXBC[SF_Pixel]         = DXBC_PS;
+	Entry.StageDXBC[SF_Geometry]      = DXBC_GS;
+	Entry.StageDXBC[SF_Mesh]          = DXBC_MS;
+	Entry.StageDXBC[SF_Amplification] = DXBC_AS;
+
+	// Store root signature blob for the direct SQLite path
+	if (RootSignatureBlob)
+	{
+		Entry.RootSigBlob = RootSignatureBlob;
+	}
+
+	// Build StreamData for the D3D12 API path
+	// Layout: [RootSig subobject] [VS] [PS] [GS] [MS] [AS] -- traditional raster or mesh shader pipeline
+	// Each subobject: { D3D12_PIPELINE_STATE_SUBOBJECT_TYPE (UINT), aligned payload }
+	if (RootSignatureBlob)
+	{
+		struct FRootSigSubobject
+		{
+			D3D12_PIPELINE_STATE_SUBOBJECT_TYPE Type;
+			D3D12_SERIALIZED_ROOT_SIGNATURE_DESC RootSig;
+		};
+		struct FShaderSubobject
+		{
+			D3D12_PIPELINE_STATE_SUBOBJECT_TYPE Type;
+			D3D12_SHADER_BYTECODE Shader;
+		};
+
+		// Point at the entry-owned StageDXBC copies rather than the caller's parameter arrays,
+		// so that StreamData's shader bytecode pointers remain valid after this function returns.
+		struct FStageDesc { D3D12_PIPELINE_STATE_SUBOBJECT_TYPE Type; const TArray<uint8>* Data; };
+		const FStageDesc Stages[] = {
+			{ D3D12_PIPELINE_STATE_SUBOBJECT_TYPE_VS, &Entry.StageDXBC[SF_Vertex]        },
+			{ D3D12_PIPELINE_STATE_SUBOBJECT_TYPE_PS, &Entry.StageDXBC[SF_Pixel]         },
+			{ D3D12_PIPELINE_STATE_SUBOBJECT_TYPE_GS, &Entry.StageDXBC[SF_Geometry]      },
+			{ D3D12_PIPELINE_STATE_SUBOBJECT_TYPE_MS, &Entry.StageDXBC[SF_Mesh]          },
+			{ D3D12_PIPELINE_STATE_SUBOBJECT_TYPE_AS, &Entry.StageDXBC[SF_Amplification] },
+		};
+
+		SIZE_T StreamSize = sizeof(FRootSigSubobject);
+		for (const FStageDesc& Stage : Stages)
+		{
+			if (!Stage.Data->IsEmpty())
+			{
+				StreamSize += sizeof(FShaderSubobject);
+			}
+		}
+
+		Entry.StreamData.SetNumZeroed((int32)StreamSize);
+		uint8* Ptr = Entry.StreamData.GetData();
+
+		FRootSigSubobject RootSigSub = {};
+		RootSigSub.Type = D3D12_PIPELINE_STATE_SUBOBJECT_TYPE_SERIALIZED_ROOT_SIGNATURE;
+		RootSigSub.RootSig.pSerializedBlob = RootSignatureBlob->GetBufferPointer();
+		RootSigSub.RootSig.SerializedBlobSizeInBytes = RootSignatureBlob->GetBufferSize();
+		FMemory::Memcpy(Ptr, &RootSigSub, sizeof(FRootSigSubobject));
+		Ptr += sizeof(FRootSigSubobject);
+
+		for (const FStageDesc& Stage : Stages)
+		{
+			if (Stage.Data->IsEmpty())
+			{
+				continue;
+			}
+			FShaderSubobject ShaderSub = {};
+			ShaderSub.Type = Stage.Type;
+			ShaderSub.Shader.pShaderBytecode = Stage.Data->GetData();
+			ShaderSub.Shader.BytecodeLength = (SIZE_T)Stage.Data->Num();
+			FMemory::Memcpy(Ptr, &ShaderSub, sizeof(FShaderSubobject));
+			Ptr += sizeof(FShaderSubobject);
+		}
+	}
+
+	{
+		FScopeLock Lock(&EntriesLock);
+		CollectedEntries.Add(MoveTemp(Entry));
+	}
+
+	return true;
+}
+
+//--------------------------------------------------------------------------------------------------
+// FSODBFactory
+//--------------------------------------------------------------------------------------------------
+
+FSODBFactory::~FSODBFactory()
+{
+	if (Factory)
+	{
+		Factory->Release();
+		Factory = nullptr;
+	}
+	if (D3D12Module)
+	{
+		FreeLibrary(D3D12Module);
+		D3D12Module = nullptr;
+	}
+}
+
+bool FSODBFactory::Initialize()
+{
+	// We load D3D12Core.dll directly (rather than going through D3D12.dll) because
+	// ID3D12StateObjectDatabaseFactory is only exposed via D3D12GetInterface in the
+	// Agility SDK's D3D12Core.dll, not in the OS-provided D3D12.dll.
+	//
+	// The path is resolved via FindAgilitySDKFile -- the same helper used for
+	// D3D12StateObjectCompiler.exe and its DLL -- keeping all Agility SDK file
+	// resolution on a single code path. Note: this is intentionally NOT the
+	// OS-loader path (D3D12SDKPath export in UnrealAgilitySDKLink.inl), which only
+	// applies when loading D3D12.dll itself.
+	const FString D3D12CorePath = FindAgilitySDKFile(TEXT("D3D12Core.dll"));
+	if (D3D12CorePath.IsEmpty())
+	{
+		UE_LOG(LogASDTool, Error, TEXT("D3D12Core.dll not found in AgilitySDK directories. Ensure Agility SDK 1.618+ is available."));
+		return false;
+	}
+
+	D3D12Module = LoadLibrary(*D3D12CorePath);
+	if (!D3D12Module)
+	{
+		UE_LOG(LogASDTool, Error, TEXT("Failed to load D3D12Core.dll from: %s"), *D3D12CorePath);
+		return false;
+	}
+
+	auto GetInterfaceFunc = reinterpret_cast<PFN_D3D12_GET_INTERFACE>(reinterpret_cast<void*>(GetProcAddress(D3D12Module, "D3D12GetInterface")));
+	if (!GetInterfaceFunc)
+	{
+		UE_LOG(LogASDTool, Error, TEXT("D3D12GetInterface not found in D3D12Core.dll"));
+		FreeLibrary(D3D12Module);
+		D3D12Module = nullptr;
+		return false;
+	}
+
+	// IID_ID3D12StateObjectDatabaseFactory is the correct interface identifier for D3D12GetInterface.
+	// (D3D12GetInterface takes an IID, not a CLSID -- the naming is misleading but intentional.)
+	HRESULT hr = GetInterfaceFunc(CLSID_D3D12StateObjectFactory, IID_PPV_ARGS(&Factory));
+	if (FAILED(hr) || !Factory)
+	{
+		UE_LOG(LogASDTool, Error, TEXT("Failed to create ID3D12StateObjectDatabaseFactory: 0x%08X"), hr);
+		FreeLibrary(D3D12Module);
+		D3D12Module = nullptr;
+		return false;
+	}
+
+	return true;
+}
+
+ID3D12StateObjectDatabase* FSODBFactory::CreateDatabase(const FString& OutputPath)
+{
+	ID3D12StateObjectDatabase* Database = nullptr;
+	HRESULT hr = Factory->CreateStateObjectDatabaseFromFile(
+		*OutputPath,
+		D3D12_STATE_OBJECT_DATABASE_FLAG_NONE,
+		IID_PPV_ARGS(&Database));
+
+	if (FAILED(hr) || !Database)
+	{
+		UE_LOG(LogASDTool, Error, TEXT("Failed to create SODB file '%s': 0x%08X"), *OutputPath, hr);
+		return nullptr;
+	}
+
+	return Database;
+}
+
+//--------------------------------------------------------------------------------------------------
+// FD3D12SODBFile
+//--------------------------------------------------------------------------------------------------
+
+FD3D12SODBFile::~FD3D12SODBFile()
+{
+	Close();
+}
+
+bool FD3D12SODBFile::Open(const FString& Path, ESODBOpenMode Mode)
+{
+	FilePath = Path;
+
+	if (!Factory)
+	{
+		UE_LOG(LogASDTool, Error, TEXT("FD3D12SODBFile::Open  -- no factory set. Call SetFactory() before Open()."));
+		return false;
+	}
+
+	Database = Factory->CreateDatabase(Path);
+	return Database != nullptr;
+}
+
+void FD3D12SODBFile::SetApplicationDesc(const D3D12_APPLICATION_DESC& AppDesc)
+{
+	if (Database)
+	{
+		Database->SetApplicationDesc(&AppDesc);
+	}
+}
+
+template<typename T>
+static void ManifestAppend(TArray<uint8>& Buffer, const T& Value)
+{
+	Buffer.AddUninitialized(sizeof(T));
+	FMemory::Memcpy(Buffer.GetData() + Buffer.Num() - sizeof(T), &Value, sizeof(T));
+}
+
+bool FD3D12SODBFile::WriteEntries(const TArray<FSODBEntry>& Entries)
+{
+	if (!Database)
+	{
+		return false;
+	}
+
+	bool bAllSucceeded = true;
+	for (const FSODBEntry& Entry : Entries)
+	{
+		D3D12_PIPELINE_STATE_STREAM_DESC StreamDesc = Entry.GetDesc();
+
+		int64 StartTime = FPlatformTime::Cycles64();
+		HRESULT hr = Database->StorePipelineStateDesc(
+			Entry.Key.GetData(), Entry.Key.Num(),
+			Entry.Version,
+			&StreamDesc);
+		FPlatformAtomics::InterlockedAdd(&WriteCycles, FPlatformTime::Cycles64() - StartTime);
+
+		if (FAILED(hr))
+		{
+			bAllSucceeded = false;
+		}
+	}
+
+	// Write key manifest file (<sodb_path>.keys) containing all keys+versions
+	{
+		FString KeysPath = FilePath + TEXT(".keys");
+		TArray<uint8> ManifestData;
+
+		int32 Count = Entries.Num();
+		ManifestAppend(ManifestData, Count);
+
+		for (const FSODBEntry& Entry : Entries)
+		{
+			int32 KeySize = Entry.Key.Num();
+			ManifestAppend(ManifestData, KeySize);
+			ManifestData.Append(Entry.Key.GetData(), KeySize);
+			ManifestAppend(ManifestData, Entry.Version);
+		}
+
+		if (!FFileHelper::SaveArrayToFile(ManifestData, *KeysPath))
+		{
+			UE_LOG(LogASDTool, Warning, TEXT("Failed to write key manifest: %s (SODB read will fall back to SQLite)"), *KeysPath);
+		}
+	}
+
+	return bAllSucceeded;
+}
+
+template<typename T>
+static bool ManifestRead(const uint8*& Ptr, const uint8* End, T& OutValue)
+{
+	if (Ptr + sizeof(T) > End)
+	{
+		return false;
+	}
+	FMemory::Memcpy(&OutValue, Ptr, sizeof(T));
+	Ptr += sizeof(T);
+	return true;
+}
+
+bool FD3D12SODBFile::ReadEntries(TArray<FSODBEntry>& OutEntries)
+{
+	if (!Database)
+	{
+		return false;
+	}
+
+	// Try to read key manifest first
+	FString KeysPath = FilePath + TEXT(".keys");
+	TArray<uint8> ManifestData;
+
+	if (FFileHelper::LoadFileToArray(ManifestData, *KeysPath) && ManifestData.Num() >= sizeof(int32))
+	{
+		const uint8* Ptr = ManifestData.GetData();
+		const uint8* End = ManifestData.GetData() + ManifestData.Num();
+
+		int32 Count = 0;
+		if (!ManifestRead(Ptr, End, Count))
+		{
+			return false;
+		}
+
+		for (int32 i = 0; i < Count; ++i)
+		{
+			int32 KeySize = 0;
+			if (!ManifestRead(Ptr, End, KeySize))
+			{
+				break;
+			}
+
+			// Validate KeySize before allocating -- a corrupt manifest could have a
+			// negative or huge value that would cause OOB reads or giant allocations.
+			const int64 Remaining = End - Ptr;
+			if (KeySize < 0 || int64(KeySize) + int64(sizeof(uint32)) > Remaining)
+			{
+				break;
+			}
+
+			FSODBEntry& Entry = OutEntries.Emplace_GetRef();
+			Entry.Key.SetNumUninitialized(KeySize);
+			FMemory::Memcpy(Entry.Key.GetData(), Ptr, KeySize);
+			Ptr += KeySize;
+
+			if (!ManifestRead(Ptr, End, Entry.Version))
+			{
+				break;
+			}
+
+			// Retrieve the stream data from the database
+			struct FCallback
+			{
+				static void __stdcall OnEntry(const void* pKey, UINT KeySz, UINT Ver, const D3D12_PIPELINE_STATE_STREAM_DESC* pDesc, void* pContext)
+				{
+					FSODBEntry* OutEntry = static_cast<FSODBEntry*>(pContext);
+					if (pDesc && pDesc->SizeInBytes > 0 && pDesc->pPipelineStateSubobjectStream)
+					{
+						OutEntry->StreamData.SetNumUninitialized(pDesc->SizeInBytes);
+						FMemory::Memcpy(OutEntry->StreamData.GetData(), pDesc->pPipelineStateSubobjectStream, pDesc->SizeInBytes);
+					}
+				}
+			};
+
+			Database->FindPipelineStateDesc(
+				Entry.Key.GetData(), Entry.Key.Num(),
+				&FCallback::OnEntry,
+				&Entry);
+		}
+
+		UE_LOG(LogASDTool, Display, TEXT("  Read %d PSO entries from SODB (via manifest): %s"), OutEntries.Num(), *FPaths::GetCleanFilename(FilePath));
+		return true;
+	}
+
+	// No manifest file  -- treat SODB as empty
+	UE_LOG(LogASDTool, Warning, TEXT("No key manifest found for SODB: %s (no .keys file)"), *FPaths::GetCleanFilename(FilePath));
+	return true;
+}
+
+void FD3D12SODBFile::Close()
+{
+	if (Database)
+	{
+		Database->Release();
+		Database = nullptr;
+	}
+}
+
+double FD3D12SODBFile::GetWriteTimeSeconds() const
+{
+	return FPlatformTime::ToSeconds64(WriteCycles);
+}
+
+FDirectSODBFile::~FDirectSODBFile()
+{
+	Close();
+}
+
+static void HashToKey(const uint8* Data, int64 Size, uint8 OutKey[16])
+{
+	FMD5 MD5;
+	MD5.Update(Data, Size);
+	MD5.Final(OutKey);
+}
+
+void FDirectSODBFile::ExecSQL(const char* SQL)
+{
+	char* ErrMsg = nullptr;
+	int32 Result = sqlite3_exec(DB, SQL, nullptr, nullptr, &ErrMsg);
+	if (Result != SQLITE_OK)
+	{
+		UE_LOG(LogASDTool, Error, TEXT("SQLite error: %s (SQL: %s)"), UTF8_TO_TCHAR(ErrMsg), UTF8_TO_TCHAR(SQL));
+		sqlite3_free(ErrMsg);
+	}
+}
+
+void FDirectSODBFile::PrepareSQLite(const char* SQL, sqlite3_stmt * * OutStmt)
+{
+	int32 Result = sqlite3_prepare_v2(DB, SQL, -1, OutStmt, nullptr);
+	if (Result != SQLITE_OK)
+	{
+		UE_LOG(LogASDTool, Error, TEXT("Failed to prepare SQL: %s (%s)"), UTF8_TO_TCHAR(SQL), UTF8_TO_TCHAR(sqlite3_errmsg(DB)));
+		*OutStmt = nullptr;
+	}
+}
+
+void FDirectSODBFile::CreateSchema()
+{
+	// Schema DDL must match D3D12StateObjectDatabase EXACTLY (including whitespace)
+	// because D3D12StateObjectCompiler parses sqlite_master DDL literally.
+	// Reference: https://microsoft.github.io/DirectX-Specs/d3d/StateObjectDatabase.html (SQL Database Schema section)
+	ExecSQL("CREATE TABLE root_signatures (Key BLOB NOT NULL PRIMARY KEY , value BLOB NOT NULL )");
+	ExecSQL("CREATE TABLE input_element_descs (Key BLOB NOT NULL PRIMARY KEY , SemanticName TEXT NOT NULL ,SemanticIndex INTEGER NOT NULL ,Format INTEGER NOT NULL ,InputSlot INTEGER NOT NULL ,AlignedByteOffset INTEGER NOT NULL ,InputSlotClass INTEGER NOT NULL ,InstanceDataStepRate INTEGER NOT NULL )");
+	ExecSQL("CREATE TABLE state_objects (Key BLOB NOT NULL PRIMARY KEY , Type INTEGER  ,NodeMask INTEGER  ,Flags INTEGER  ,AddToStateObjectParent BLOB  ,FOREIGN KEY(AddToStateObjectParent) REFERENCES state_objects(Key))");
+	ExecSQL("CREATE TABLE shader_bytecode (Key BLOB NOT NULL PRIMARY KEY , Type TEXT  ,Bytecode BLOB  )");
+	ExecSQL("CREATE TABLE exports (Key BLOB NOT NULL PRIMARY KEY , Name TEXT NOT NULL ,ExportToRename TEXT  ,Flags INTEGER NOT NULL )");
+	ExecSQL("CREATE TABLE depth_stencil_op_descs (Key BLOB NOT NULL PRIMARY KEY , StencilFailOp INTEGER NOT NULL ,StencilDepthFailOp INTEGER NOT NULL ,StencilPassOp INTEGER NOT NULL ,StencilFunc INTEGER NOT NULL ,StencilReadMask INTEGER NOT NULL ,StencilWriteMask INTEGER NOT NULL )");
+	ExecSQL("CREATE TABLE depth_stencil_descs (Key BLOB NOT NULL PRIMARY KEY , DepthEnable INTEGER NOT NULL ,DepthWriteMask INTEGER NOT NULL ,DepthFunc INTEGER NOT NULL ,StencilEnable INTEGER NOT NULL ,FrontFace BLOB NOT NULL ,BackFace BLOB NOT NULL ,DepthBoundsTestEnable INTEGER NOT NULL ,FOREIGN KEY(FrontFace) REFERENCES depth_stencil_op_descs(Key),FOREIGN KEY(BackFace) REFERENCES depth_stencil_op_descs(Key))");
+	ExecSQL("CREATE TABLE render_target_formats (Key BLOB NOT NULL PRIMARY KEY , RTFormat0 INTEGER NOT NULL ,RTFormat1 INTEGER NOT NULL ,RTFormat2 INTEGER NOT NULL ,RTFormat3 INTEGER NOT NULL ,RTFormat4 INTEGER NOT NULL ,RTFormat5 INTEGER NOT NULL ,RTFormat6 INTEGER NOT NULL ,RTFormat7 INTEGER NOT NULL ,NumRenderTargets INTEGER NOT NULL )");
+	ExecSQL("CREATE TABLE render_target_blend_descs (Key BLOB NOT NULL PRIMARY KEY , BlendEnable INTEGER NOT NULL ,LogicOpEnable INTEGER NOT NULL ,SrcBlend INTEGER NOT NULL ,DestBlend INTEGER NOT NULL ,BlendOp INTEGER NOT NULL ,SrcBlendAlpha INTEGER NOT NULL ,DestBlendAlpha INTEGER NOT NULL ,BlendOpAlpha INTEGER NOT NULL ,LogicOp INTEGER NOT NULL ,RenderTargetWriteMask INTEGER NOT NULL )");
+	ExecSQL("CREATE TABLE blend_descs (Key BLOB NOT NULL PRIMARY KEY , AlphaToCoverageEnable INTEGER NOT NULL ,IndependentBlendEnable INTEGER NOT NULL ,RenderTarget0 BLOB  ,RenderTarget1 BLOB  ,RenderTarget2 BLOB  ,RenderTarget3 BLOB  ,RenderTarget4 BLOB  ,RenderTarget5 BLOB  ,RenderTarget6 BLOB  ,RenderTarget7 BLOB  ,FOREIGN KEY(RenderTarget0) REFERENCES render_target_blend_descs(Key),FOREIGN KEY(RenderTarget1) REFERENCES render_target_blend_descs(Key),FOREIGN KEY(RenderTarget2) REFERENCES render_target_blend_descs(Key),FOREIGN KEY(RenderTarget3) REFERENCES render_target_blend_descs(Key),FOREIGN KEY(RenderTarget4) REFERENCES render_target_blend_descs(Key),FOREIGN KEY(RenderTarget5) REFERENCES render_target_blend_descs(Key),FOREIGN KEY(RenderTarget6) REFERENCES render_target_blend_descs(Key),FOREIGN KEY(RenderTarget7) REFERENCES render_target_blend_descs(Key))");
+	ExecSQL("CREATE TABLE rasterizer_descs (Key BLOB NOT NULL PRIMARY KEY , FillMode INTEGER NOT NULL ,CullMode INTEGER NOT NULL ,FrontCounterClockwise INTEGER NOT NULL ,DepthBias REAL NOT NULL ,DepthBiasClamp REAL NOT NULL ,SlopeScaledDepthBias REAL NOT NULL ,DepthClipEnable INTEGER NOT NULL ,LineRasterizationMode INTEGER NOT NULL ,ForcedSampleCount INTEGER NOT NULL ,ConservativeRaster INTEGER NOT NULL )");
+	ExecSQL("CREATE TABLE view_instancing_descs (Key BLOB NOT NULL PRIMARY KEY , ViewInstanceCount INTEGER NOT NULL ,RenderFlags INTEGER NOT NULL ,ViewportArrayIndex0 INTEGER  ,RenderTargetArrayIndex0 INTEGER  ,ViewportArrayIndex1 INTEGER  ,RenderTargetArrayIndex1 INTEGER  ,ViewportArrayIndex2 INTEGER  ,RenderTargetArrayIndex2 INTEGER  ,ViewportArrayIndex3 INTEGER  ,RenderTargetArrayIndex3 INTEGER  )");
+	ExecSQL("CREATE TABLE so_declarations (Key BLOB NOT NULL PRIMARY KEY , Stream INTEGER NOT NULL ,SemanticName TEXT NOT NULL ,SemanticIndex INTEGER NOT NULL ,StartComponent INTEGER NOT NULL ,ComponentCount INTEGER NOT NULL ,OutputSlot INTEGER NOT NULL )");
+	ExecSQL("CREATE TABLE stream_out_descs (Key BLOB NOT NULL PRIMARY KEY , BufferStride0 INTEGER NOT NULL ,BufferStride1 INTEGER NOT NULL ,BufferStride2 INTEGER NOT NULL ,BufferStride3 INTEGER NOT NULL ,NumStrides INTEGER NOT NULL ,RasterizedStream INTEGER NOT NULL )");
+	ExecSQL("CREATE TABLE rt_hit_groups (Key BLOB NOT NULL PRIMARY KEY , HitGroupExport TEXT NOT NULL ,Type INTEGER NOT NULL ,AnyHitShaderImport TEXT  ,ClosestHitShaderImport TEXT  ,IntersectionShaderImport TEXT  )");
+	ExecSQL("CREATE TABLE rt_shader_config (Key BLOB NOT NULL PRIMARY KEY , MaxPayloadSizeInBytes INTEGER NOT NULL ,MaxAttributeSizeInBytes INTEGER NOT NULL )");
+	ExecSQL("CREATE TABLE rt_pipeline_config (Key BLOB NOT NULL PRIMARY KEY , MaxTraceRecursionDepth INTEGER NOT NULL ,Flags INTEGER NOT NULL )");
+	ExecSQL("CREATE TABLE dxil_subobject_to_exports_associations (Key BLOB NOT NULL PRIMARY KEY , SubobjectToAssociate TEXT NOT NULL )");
+	ExecSQL("CREATE TABLE subobject_to_exports_associations (Key BLOB NOT NULL PRIMARY KEY , SubobjectType INTEGER NOT NULL ,SubobjectKey BLOB NOT NULL )");
+	ExecSQL("CREATE TABLE generic_programs (Key BLOB NOT NULL PRIMARY KEY , ProgramName TEXT  ,InputLayout BLOB  ,DepthStencilDesc BLOB  ,RenderTargetFormats BLOB  ,BlendDesc BLOB  ,RasterizerDesc BLOB  ,ViewInstancingDesc BLOB  ,StreamOutDesc BLOB  ,SampleDesc_Count INTEGER  ,SampleDesc_Quality INTEGER  ,SampleMask INTEGER  ,IBStripCutValue INTEGER  ,PrimitiveTopology INTEGER  ,DSVFormat INTEGER  ,NodeMask INTEGER  ,Flags INTEGER  ,FOREIGN KEY(DepthStencilDesc) REFERENCES depth_stencil_descs(Key),FOREIGN KEY(RenderTargetFormats) REFERENCES render_target_formats(Key),FOREIGN KEY(BlendDesc) REFERENCES blend_descs(Key),FOREIGN KEY(RasterizerDesc) REFERENCES rasterizer_descs(Key),FOREIGN KEY(ViewInstancingDesc) REFERENCES view_instancing_descs(Key),FOREIGN KEY(StreamOutDesc) REFERENCES stream_out_descs(Key))");
+	ExecSQL("CREATE TABLE pipeline_states (Key BLOB NOT NULL PRIMARY KEY , RootSignature BLOB  ,InputLayout BLOB  ,ByteCode_VS BLOB  ,ByteCode_PS BLOB  ,ByteCode_HS BLOB  ,ByteCode_DS BLOB  ,ByteCode_GS BLOB  ,ByteCode_AS BLOB  ,ByteCode_MS BLOB  ,ByteCode_CS BLOB  ,DepthStencilDesc BLOB  ,RenderTargetFormats BLOB  ,BlendDesc BLOB  ,RasterizerDesc BLOB  ,ViewInstancingDesc BLOB  ,StreamOutDesc BLOB  ,SampleDesc_Count INTEGER  ,SampleDesc_Quality INTEGER  ,SampleMask INTEGER  ,IBStripCutValue INTEGER  ,PrimitiveTopology INTEGER  ,DSVFormat INTEGER  ,NodeMask INTEGER  ,Flags INTEGER  ,FOREIGN KEY(RootSignature) REFERENCES root_signatures(Key),FOREIGN KEY(ByteCode_VS) REFERENCES shader_bytecode(Key),FOREIGN KEY(ByteCode_PS) REFERENCES shader_bytecode(Key),FOREIGN KEY(ByteCode_HS) REFERENCES shader_bytecode(Key),FOREIGN KEY(ByteCode_DS) REFERENCES shader_bytecode(Key),FOREIGN KEY(ByteCode_GS) REFERENCES shader_bytecode(Key),FOREIGN KEY(ByteCode_AS) REFERENCES shader_bytecode(Key),FOREIGN KEY(ByteCode_MS) REFERENCES shader_bytecode(Key),FOREIGN KEY(ByteCode_CS) REFERENCES shader_bytecode(Key),FOREIGN KEY(DepthStencilDesc) REFERENCES depth_stencil_descs(Key),FOREIGN KEY(RenderTargetFormats) REFERENCES render_target_formats(Key),FOREIGN KEY(BlendDesc) REFERENCES blend_descs(Key),FOREIGN KEY(RasterizerDesc) REFERENCES rasterizer_descs(Key),FOREIGN KEY(ViewInstancingDesc) REFERENCES view_instancing_descs(Key),FOREIGN KEY(StreamOutDesc) REFERENCES stream_out_descs(Key))");
+	ExecSQL("CREATE TABLE work_graphs (Key BLOB NOT NULL PRIMARY KEY , ProgramName TEXT NOT NULL ,Flags INTEGER NOT NULL )");
+	ExecSQL("CREATE TABLE node_ids (Key BLOB NOT NULL PRIMARY KEY , Name TEXT NOT NULL ,ArrayIndex INTEGER NOT NULL )");
+	ExecSQL("CREATE TABLE node_output_overrides (Key BLOB NOT NULL PRIMARY KEY , OutputIndex INTEGER NOT NULL ,NewName BLOB  ,AllowSparseNodes INTEGER  ,MaxRecords INTEGER  ,MaxRecordsSharedWithOutputIndex INTEGER  ,FOREIGN KEY(NewName) REFERENCES node_ids(Key))");
+	ExecSQL("CREATE TABLE shader_nodes (Key BLOB NOT NULL PRIMARY KEY , ShaderOrProgram TEXT NOT NULL ,NodeType INTEGER NOT NULL ,OverridesType INTEGER NOT NULL ,LocalRootArgumentsTableIndex INTEGER  ,ProgramEntry INTEGER  ,NewName BLOB  ,ShareInputOf BLOB  ,DispatchGridX INTEGER  ,DispatchGridY INTEGER  ,DispatchGridZ INTEGER  ,MaxDispatchGridX INTEGER  ,MaxDispatchGridY INTEGER  ,MaxDispatchGridZ INTEGER  ,MaxInputRecordsPerGraphEntryRecord_RecordCount INTEGER  ,MaxInputRecordsPerGraphEntryRecord_bCountSharedAcrossNodeArray INTEGER  ,FOREIGN KEY(NewName) REFERENCES node_ids(Key),FOREIGN KEY(ShareInputOf) REFERENCES node_ids(Key))");
+	ExecSQL("CREATE TABLE input_layout_to_input_element_associations (InputLayoutKey BLOB NOT NULL ,InputElementKey BLOB NOT NULL ,FOREIGN KEY(InputElementKey) REFERENCES input_element_descs(Key),UNIQUE(InputLayoutKey, InputElementKey))");
+	ExecSQL("CREATE TABLE stream_output_desc_to_stream_output_decl_associations (StreamOutDescKey BLOB NOT NULL ,StreamOutDeclKey BLOB NOT NULL ,FOREIGN KEY(StreamOutDescKey) REFERENCES stream_out_descs(Key),FOREIGN KEY(StreamOutDeclKey) REFERENCES so_declarations(Key),UNIQUE(StreamOutDescKey, StreamOutDeclKey))");
+	ExecSQL("CREATE TABLE so_to_global_rs_associations (StateObjectKey BLOB NOT NULL ,RootSignatureKey BLOB NOT NULL ,FOREIGN KEY(StateObjectKey) REFERENCES state_objects(Key),FOREIGN KEY(RootSignatureKey) REFERENCES root_signatures(Key),UNIQUE(StateObjectKey, RootSignatureKey))");
+	ExecSQL("CREATE TABLE so_to_local_rs_associations (StateObjectKey BLOB NOT NULL ,RootSignatureKey BLOB NOT NULL ,FOREIGN KEY(StateObjectKey) REFERENCES state_objects(Key),FOREIGN KEY(RootSignatureKey) REFERENCES root_signatures(Key),UNIQUE(StateObjectKey, RootSignatureKey))");
+	ExecSQL("CREATE TABLE so_to_dxil_lib_associations (StateObjectKey BLOB NOT NULL ,DxilLibKey BLOB NOT NULL ,ExportKey BLOB  ,FOREIGN KEY(StateObjectKey) REFERENCES state_objects(Key),FOREIGN KEY(DxilLibKey) REFERENCES shader_bytecode(Key),FOREIGN KEY(ExportKey) REFERENCES exports(Key),UNIQUE(StateObjectKey, DxilLibKey, ExportKey))");
+	ExecSQL("CREATE TABLE so_to_existing_so_associations (StateObjectKey BLOB NOT NULL ,ExistingStateObjectKey BLOB NOT NULL ,ExportKey BLOB  ,FOREIGN KEY(StateObjectKey) REFERENCES state_objects(Key),FOREIGN KEY(ExistingStateObjectKey) REFERENCES state_objects(Key),FOREIGN KEY(ExportKey) REFERENCES exports(Key),UNIQUE(StateObjectKey, ExistingStateObjectKey, ExportKey))");
+	ExecSQL("CREATE TABLE so_to_hit_group_associations (StateObjectKey BLOB NOT NULL ,HitGroupKey BLOB NOT NULL ,FOREIGN KEY(StateObjectKey) REFERENCES state_objects(Key),FOREIGN KEY(HitGroupKey) REFERENCES rt_hit_groups(Key),UNIQUE(StateObjectKey, HitGroupKey))");
+	ExecSQL("CREATE TABLE so_to_rt_shader_config_associations (StateObjectKey BLOB NOT NULL ,ShaderConfigKey BLOB NOT NULL ,FOREIGN KEY(StateObjectKey) REFERENCES state_objects(Key),FOREIGN KEY(ShaderConfigKey) REFERENCES rt_shader_config(Key),UNIQUE(StateObjectKey, ShaderConfigKey))");
+	ExecSQL("CREATE TABLE so_to_rt_pipeline_config_associations (StateObjectKey BLOB NOT NULL ,PipelineConfigKey BLOB NOT NULL ,FOREIGN KEY(StateObjectKey) REFERENCES state_objects(Key),FOREIGN KEY(PipelineConfigKey) REFERENCES rt_pipeline_config(Key),UNIQUE(StateObjectKey, PipelineConfigKey))");
+	ExecSQL("CREATE TABLE so_to_dxil_subobject_to_exports_associations (StateObjectKey BLOB NOT NULL ,DxilSubobjectToExportsAssociationKey BLOB NOT NULL ,FOREIGN KEY(StateObjectKey) REFERENCES state_objects(Key),FOREIGN KEY(DxilSubobjectToExportsAssociationKey) REFERENCES dxil_subobject_to_exports_associations(Key),UNIQUE(StateObjectKey, DxilSubobjectToExportsAssociationKey))");
+	ExecSQL("CREATE TABLE so_to_subobject_to_exports_associations (StateObjectKey BLOB NOT NULL ,SubobjectToExportsAssociationKey BLOB NOT NULL ,FOREIGN KEY(StateObjectKey) REFERENCES state_objects(Key),FOREIGN KEY(SubobjectToExportsAssociationKey) REFERENCES subobject_to_exports_associations(Key),UNIQUE(StateObjectKey, SubobjectToExportsAssociationKey))");
+	ExecSQL("CREATE TABLE so_to_generic_program_associations (StateObjectKey BLOB NOT NULL ,GenericProgramKey BLOB NOT NULL ,FOREIGN KEY(StateObjectKey) REFERENCES state_objects(Key),FOREIGN KEY(GenericProgramKey) REFERENCES generic_programs(Key),UNIQUE(StateObjectKey, GenericProgramKey))");
+	ExecSQL("CREATE TABLE so_to_work_graph_associations (StateObjectKey BLOB NOT NULL ,WorkGraphKey BLOB NOT NULL ,FOREIGN KEY(StateObjectKey) REFERENCES state_objects(Key),FOREIGN KEY(WorkGraphKey) REFERENCES work_graphs(Key),UNIQUE(StateObjectKey, WorkGraphKey))");
+	ExecSQL("CREATE TABLE work_graph_to_entrypoint_node_id_associations (WorkGraphKey BLOB NOT NULL ,NodeIDKey BLOB NOT NULL ,FOREIGN KEY(WorkGraphKey) REFERENCES work_graphs(Key),FOREIGN KEY(NodeIDKey) REFERENCES node_ids(Key),UNIQUE(WorkGraphKey, NodeIDKey))");
+	ExecSQL("CREATE TABLE workgraph_node_to_node_output_overrides_associations (WorkGraphNodeKey BLOB NOT NULL ,NodeOutputOverridesKey BLOB NOT NULL ,FOREIGN KEY(WorkGraphNodeKey) REFERENCES shader_nodes(Key),FOREIGN KEY(NodeOutputOverridesKey) REFERENCES node_output_overrides(Key),UNIQUE(WorkGraphNodeKey, NodeOutputOverridesKey))");
+	ExecSQL("CREATE TABLE work_graph_to_work_graph_node_associations (WorkGraphKey BLOB NOT NULL ,WorkGraphNodeKey BLOB NOT NULL ,FOREIGN KEY(WorkGraphKey) REFERENCES work_graphs(Key),FOREIGN KEY(WorkGraphNodeKey) REFERENCES shader_nodes(Key),UNIQUE(WorkGraphKey, WorkGraphNodeKey))");
+	ExecSQL("CREATE TABLE string_associations (OwningTableKey BLOB  ,Value TEXT NOT NULL ,UNIQUE(OwningTableKey, Value))");
+	ExecSQL("CREATE TABLE groups (Key BLOB NOT NULL PRIMARY KEY , Version INTEGER NOT NULL ,PSOKey BLOB  ,SOKey BLOB  ,FOREIGN KEY(PSOKey) REFERENCES pipeline_states(Key),FOREIGN KEY(SOKey) REFERENCES state_objects(Key))");
+	ExecSQL("CREATE TABLE app_id (id INTEGER NOT NULL PRIMARY KEY , exe TEXT NOT NULL ,app_name TEXT NOT NULL ,engine_name TEXT  ,app_version INTEGER NOT NULL ,engine_version INTEGER NOT NULL )");
+	ExecSQL("CREATE UNIQUE INDEX so_to_dxil_lib_associations_partial_index ON so_to_dxil_lib_associations(StateObjectKey,DxilLibKey) WHERE ExportKey IS NULL");
+	ExecSQL("CREATE UNIQUE INDEX so_to_existing_so_associations_partial_index ON so_to_existing_so_associations(StateObjectKey,ExistingStateObjectKey) WHERE ExportKey IS NULL");
+	ExecSQL("CREATE UNIQUE INDEX string_associations_partial_index ON string_associations(Value) WHERE OwningTableKey IS NULL");
+}
+
+void FDirectSODBFile::Cleanup()
+{
+	// Finalize the 4 data-write statements. StmtGroup and DB are managed by Close() directly
+	// since they need to be finalized in a specific order around the app_id insert and COMMIT.
+	if (StmtRootSig)
+	{
+		sqlite3_finalize(StmtRootSig);
+		StmtRootSig = nullptr;
+	}
+	if (StmtBytecode)
+	{
+		sqlite3_finalize(StmtBytecode);
+		StmtBytecode = nullptr;
+	}
+	if (StmtPSO)
+	{
+		sqlite3_finalize(StmtPSO);
+		StmtPSO = nullptr;
+	}
+	if (StmtPSO_Graphics)
+	{
+		sqlite3_finalize(StmtPSO_Graphics);
+		StmtPSO_Graphics = nullptr;
+	}
+}
+
+bool FDirectSODBFile::Open(const FString& Path, ESODBOpenMode Mode)
+{
+	FilePath = Path;
+	OpenMode = Mode;
+
+	sqlite3_initialize();
+
+	if (Mode == ESODBOpenMode::Write)
+	{
+		// Delete existing file if present
+		if (FPaths::FileExists(Path))
+		{
+			IFileManager::Get().Delete(*Path);
+		}
+
+		int32 Result = sqlite3_open(TCHAR_TO_UTF8(*Path), &DB);
+		if (Result != SQLITE_OK || !DB)
+		{
+			UE_LOG(LogASDTool, Error, TEXT("Failed to open SQLite database '%s': %s"), *Path, UTF8_TO_TCHAR(sqlite3_errmsg(DB)));
+			if (DB)
+			{
+				sqlite3_close(DB);
+				DB = nullptr;
+			}
+			return false;
+		}
+
+		// Set pragmas to match official SODB format
+		ExecSQL("PRAGMA user_version = 2");
+		ExecSQL("PRAGMA application_id = 222122203"); // 0xD3D50DB
+		ExecSQL("PRAGMA synchronous = OFF");
+
+		// Create all tables matching the official SODB schema
+		CreateSchema();
+
+		// Begin single transaction for all inserts
+		ExecSQL("BEGIN TRANSACTION");
+
+		// Prepare statements for the tables we actually write to
+		PrepareSQLite("INSERT OR IGNORE INTO root_signatures (Key, value) VALUES (?, ?)", &StmtRootSig);
+		PrepareSQLite("INSERT OR IGNORE INTO shader_bytecode (Key, Type, Bytecode) VALUES (?, ?, ?)", &StmtBytecode);
+		PrepareSQLite("INSERT INTO pipeline_states (Key, RootSignature, ByteCode_CS) VALUES (?, ?, ?)", &StmtPSO);
+		PrepareSQLite(
+			"INSERT INTO pipeline_states (Key, RootSignature, ByteCode_VS, ByteCode_PS, ByteCode_GS, ByteCode_MS, ByteCode_AS, "
+			"DepthStencilDesc, RenderTargetFormats, BlendDesc, RasterizerDesc, "
+			"SampleDesc_Count, SampleDesc_Quality, SampleMask, IBStripCutValue, PrimitiveTopology, DSVFormat, NodeMask, Flags) "
+			"VALUES (?, ?, ?, ?, ?, ?, ?, NULL, NULL, NULL, NULL, NULL, NULL, NULL, NULL, NULL, NULL, NULL, NULL)",
+			&StmtPSO_Graphics);
+		PrepareSQLite("INSERT INTO groups (Key, Version, PSOKey) VALUES (?, 0, ?)", &StmtGroup);
+
+		if (!StmtRootSig || !StmtBytecode || !StmtPSO || !StmtPSO_Graphics || !StmtGroup)
+		{
+			UE_LOG(LogASDTool, Error, TEXT("Failed to prepare one or more SQLite statements"));
+			Close();
+			return false;
+		}
+	}
+	else
+	{
+		int32 Result = sqlite3_open_v2(TCHAR_TO_UTF8(*Path), &DB, SQLITE_OPEN_READONLY, nullptr);
+		if (Result != SQLITE_OK || !DB)
+		{
+			UE_LOG(LogASDTool, Error, TEXT("Failed to open SQLite database '%s' for reading: %s"), *Path, UTF8_TO_TCHAR(sqlite3_errmsg(DB)));
+			if (DB)
+			{
+				sqlite3_close(DB);
+				DB = nullptr;
+			}
+			return false;
+		}
+	}
+
+	return true;
+}
+
+void FDirectSODBFile::SetApplicationDesc(const D3D12_APPLICATION_DESC& AppDesc)
+{
+	// Cache the app desc strings  -- the pointers in AppDesc may not remain valid
+	CachedAppDesc = AppDesc;
+	CachedExeFilename = AppDesc.pExeFilename ? AppDesc.pExeFilename : TEXT("");
+	CachedAppName = AppDesc.pName ? AppDesc.pName : TEXT("");
+	CachedEngineName = AppDesc.pEngineName ? AppDesc.pEngineName : TEXT("");
+}
+
+/** Helper: insert a shader bytecode row and return the 16-byte MD5 key. Returns false if DXBC is empty. */
+static bool InsertShaderBytecode(sqlite3_stmt* Stmt, const TArray<uint8>& DXBCData, const char* TypeStr, uint8 OutKey[16])
+{
+	if (DXBCData.IsEmpty())
+	{
+		return false;
+	}
+	HashToKey(DXBCData.GetData(), DXBCData.Num(), OutKey);
+	sqlite3_bind_blob(Stmt, 1, OutKey, 16, SQLITE_STATIC);
+	sqlite3_bind_text(Stmt, 2, TypeStr, -1, SQLITE_STATIC);
+	sqlite3_bind_blob(Stmt, 3, DXBCData.GetData(), (int)DXBCData.Num(), SQLITE_STATIC);
+	const int32 StepResult = sqlite3_step(Stmt);
+	sqlite3_reset(Stmt);
+	// INSERT OR IGNORE -- SQLITE_DONE = inserted, SQLITE_CONSTRAINT = duplicate (both are fine)
+	if (StepResult != SQLITE_DONE && StepResult != SQLITE_CONSTRAINT)
+	{
+		UE_LOG(LogASDTool, Warning, TEXT("InsertShaderBytecode: unexpected sqlite3_step result %d for type %s"), StepResult, UTF8_TO_TCHAR(TypeStr));
+	}
+	return true;
+}
+
+bool FDirectSODBFile::WriteEntries(const TArray<FSODBEntry>& Entries)
+{
+	int64 StartTime = FPlatformTime::Cycles64();
+
+	for (const FSODBEntry& Entry : Entries)
+	{
+		// Insert root signature (OR IGNORE deduplicates)
+		uint8 RootSigKey[16] = {};
+		if (Entry.RootSigBlob)
+		{
+			HashToKey((const uint8*)Entry.RootSigBlob->GetBufferPointer(), Entry.RootSigBlob->GetBufferSize(), RootSigKey);
+			sqlite3_bind_blob(StmtRootSig, 1, RootSigKey, 16, SQLITE_STATIC);
+			sqlite3_bind_blob(StmtRootSig, 2, Entry.RootSigBlob->GetBufferPointer(), (int)Entry.RootSigBlob->GetBufferSize(), SQLITE_STATIC);
+			const int32 RootSigResult = sqlite3_step(StmtRootSig);
+			sqlite3_reset(StmtRootSig);
+			// INSERT OR IGNORE -- SQLITE_DONE = inserted, SQLITE_CONSTRAINT = duplicate (both fine)
+			if (RootSigResult != SQLITE_DONE && RootSigResult != SQLITE_CONSTRAINT)
+			{
+				UE_LOG(LogASDTool, Warning, TEXT("Failed to insert root signature (sqlite3_step: %d)"), RootSigResult);
+			}
+		}
+
+		// Use the full 8-byte FXxHash64 key for both PSO and group entries
+		const uint8* PSOKey = Entry.Key.GetData();
+		const int    PSOKeySize = Entry.Key.Num();
+		const void* RootSigKeyPtr = Entry.RootSigBlob ? RootSigKey : nullptr;
+		const int RootSigKeySize = Entry.RootSigBlob ? 16 : 0;
+
+		int32 StepResult = SQLITE_DONE;
+
+		if (Entry.bIsGraphics)
+		{
+			// Insert per-stage shader bytecode rows
+			uint8 KeyVS[16] = {}, KeyPS[16] = {}, KeyGS[16] = {}, KeyMS[16] = {}, KeyAS[16] = {};
+			bool bHasVS = InsertShaderBytecode(StmtBytecode, Entry.StageDXBC[SF_Vertex],        "VS", KeyVS);
+			bool bHasPS = InsertShaderBytecode(StmtBytecode, Entry.StageDXBC[SF_Pixel],         "PS", KeyPS);
+			bool bHasGS = InsertShaderBytecode(StmtBytecode, Entry.StageDXBC[SF_Geometry],      "GS", KeyGS);
+			bool bHasMS = InsertShaderBytecode(StmtBytecode, Entry.StageDXBC[SF_Mesh],          "MS", KeyMS);
+			bool bHasAS = InsertShaderBytecode(StmtBytecode, Entry.StageDXBC[SF_Amplification], "AS", KeyAS);
+
+			// Insert graphics pipeline state
+			sqlite3_bind_blob(StmtPSO_Graphics, 1, PSOKey, PSOKeySize, SQLITE_STATIC);
+			sqlite3_bind_blob(StmtPSO_Graphics, 2, RootSigKeyPtr, RootSigKeySize, SQLITE_STATIC);
+			sqlite3_bind_blob(StmtPSO_Graphics, 3, bHasVS ? KeyVS : nullptr, bHasVS ? 16 : 0, SQLITE_STATIC);
+			sqlite3_bind_blob(StmtPSO_Graphics, 4, bHasPS ? KeyPS : nullptr, bHasPS ? 16 : 0, SQLITE_STATIC);
+				sqlite3_bind_blob(StmtPSO_Graphics, 5, bHasGS ? KeyGS : nullptr, bHasGS ? 16 : 0, SQLITE_STATIC);
+				sqlite3_bind_blob(StmtPSO_Graphics, 6, bHasMS ? KeyMS : nullptr, bHasMS ? 16 : 0, SQLITE_STATIC);
+				sqlite3_bind_blob(StmtPSO_Graphics, 7, bHasAS ? KeyAS : nullptr, bHasAS ? 16 : 0, SQLITE_STATIC);
+			StepResult = sqlite3_step(StmtPSO_Graphics);
+			sqlite3_reset(StmtPSO_Graphics);
+		}
+		else
+		{
+			// Compute path  -- insert single CS bytecode row
+			uint8 BytecodeKey[16];
+			HashToKey(Entry.StageDXBC[SF_Compute].GetData(), Entry.StageDXBC[SF_Compute].Num(), BytecodeKey);
+			sqlite3_bind_blob(StmtBytecode, 1, BytecodeKey, 16, SQLITE_STATIC);
+			sqlite3_bind_text(StmtBytecode, 2, "CS", -1, SQLITE_STATIC);
+			sqlite3_bind_blob(StmtBytecode, 3, Entry.StageDXBC[SF_Compute].GetData(), (int)Entry.StageDXBC[SF_Compute].Num(), SQLITE_STATIC);
+			int32 BytecodeResult = sqlite3_step(StmtBytecode);
+			sqlite3_reset(StmtBytecode);
+			if (BytecodeResult != SQLITE_DONE)
+			{
+				UE_LOG(LogASDTool, Warning, TEXT("Failed to insert shader bytecode"));
+			}
+
+			sqlite3_bind_blob(StmtPSO, 1, PSOKey, PSOKeySize, SQLITE_STATIC);
+			sqlite3_bind_blob(StmtPSO, 2, RootSigKeyPtr, RootSigKeySize, SQLITE_STATIC);
+			sqlite3_bind_blob(StmtPSO, 3, BytecodeKey, 16, SQLITE_STATIC);
+			StepResult = sqlite3_step(StmtPSO);
+			sqlite3_reset(StmtPSO);
+		}
+
+		if (StepResult != SQLITE_DONE)
+		{
+			UE_LOG(LogASDTool, Warning, TEXT("Failed to insert pipeline state"));
+		}
+
+		// Insert group entry -- key and PSO key are both the full 8-byte FXxHash64
+		sqlite3_bind_blob(StmtGroup, 1, Entry.Key.GetData(), Entry.Key.Num(), SQLITE_STATIC);
+		sqlite3_bind_blob(StmtGroup, 2, PSOKey, PSOKeySize, SQLITE_STATIC);
+		StepResult = sqlite3_step(StmtGroup);
+		sqlite3_reset(StmtGroup);
+		if (StepResult != SQLITE_DONE)
+		{
+			UE_LOG(LogASDTool, Warning, TEXT("Failed to insert group entry"));
+		}
+	}
+
+	FPlatformAtomics::InterlockedAdd(&WriteCycles, FPlatformTime::Cycles64() - StartTime);
+	return true;
+}
+
+bool FDirectSODBFile::ReadEntries(TArray<FSODBEntry>& OutEntries)
+{
+	// Read from groups + pipeline_states tables
+	sqlite3_stmt* Stmt = nullptr;
+	PrepareSQLite(
+		"SELECT g.Key, g.Version, ps.ByteCode_CS, sb.Bytecode, ps.RootSignature, rs.value "
+		"FROM groups g "
+		"LEFT JOIN pipeline_states ps ON ps.Key = g.PSOKey "
+		"LEFT JOIN shader_bytecode sb ON sb.Key = ps.ByteCode_CS "
+		"LEFT JOIN root_signatures rs ON rs.Key = ps.RootSignature",
+		&Stmt);
+
+	if (!Stmt)
+	{
+		return false;
+	}
+
+	while (sqlite3_step(Stmt) == SQLITE_ROW)
+	{
+		FSODBEntry& Entry = OutEntries.Emplace_GetRef();
+
+		// Key (group key)
+		int32 KeySize = sqlite3_column_bytes(Stmt, 0);
+		const void* KeyData = sqlite3_column_blob(Stmt, 0);
+		if (KeyData && KeySize > 0)
+		{
+			Entry.Key.SetNumUninitialized(KeySize);
+			FMemory::Memcpy(Entry.Key.GetData(), KeyData, KeySize);
+		}
+
+		// Version
+		Entry.Version = (uint32)sqlite3_column_int(Stmt, 1);
+
+		// DXBC data (shader bytecode)
+		int32 DXBCSize = sqlite3_column_bytes(Stmt, 3);
+		const void* DXBCPtr = sqlite3_column_blob(Stmt, 3);
+		if (DXBCPtr && DXBCSize > 0)
+		{
+			Entry.StageDXBC[SF_Compute].SetNumUninitialized(DXBCSize);
+			FMemory::Memcpy(Entry.StageDXBC[SF_Compute].GetData(), DXBCPtr, DXBCSize);
+		}
+
+		// Root signature raw bytes
+		int32 RootSigSize = sqlite3_column_bytes(Stmt, 5);
+		const void* RootSigPtr = sqlite3_column_blob(Stmt, 5);
+		if (RootSigPtr && RootSigSize > 0)
+		{
+			Entry.RootSigData.SetNumUninitialized(RootSigSize);
+			FMemory::Memcpy(Entry.RootSigData.GetData(), RootSigPtr, RootSigSize);
+		}
+
+		// NOTE: StreamData is intentionally left empty for the direct SQLite path.
+		// The caller must build the pipeline state stream from DXBCData/RootSigBlob
+		// with live pointers at compile time.
+	}
+
+	sqlite3_finalize(Stmt);
+
+	UE_LOG(LogASDTool, Display, TEXT("  Read %d PSO entries from SODB (SQLite): %s"), OutEntries.Num(), *FPaths::GetCleanFilename(FilePath));
+	return true;
+}
+
+void FDirectSODBFile::Close()
+{
+	if (!DB)
+	{
+		return;
+	}
+
+	if (OpenMode == ESODBOpenMode::Write)
+	{
+		// Finalize write statements before app_id insert -- releases resources held by open statements
+		Cleanup();
+		if (StmtGroup)
+		{
+			sqlite3_finalize(StmtGroup);
+			StmtGroup = nullptr;
+		}
+
+		// Write app_id using cached desc
+		const D3D12_APPLICATION_DESC& AppDesc = CachedAppDesc;
+		uint64 AppVersion = ((uint64)AppDesc.Version.VersionParts[0] << 48) | ((uint64)AppDesc.Version.VersionParts[1] << 32)
+			| ((uint64)AppDesc.Version.VersionParts[2] << 16) | (uint64)AppDesc.Version.VersionParts[3];
+		uint64 EngineVersion = ((uint64)AppDesc.EngineVersion.VersionParts[0] << 48) | ((uint64)AppDesc.EngineVersion.VersionParts[1] << 32)
+			| ((uint64)AppDesc.EngineVersion.VersionParts[2] << 16) | (uint64)AppDesc.EngineVersion.VersionParts[3];
+
+		sqlite3_stmt* AppStmt = nullptr;
+		PrepareSQLite("INSERT INTO app_id (id, exe, app_name, engine_name, app_version, engine_version) VALUES (0, ?, ?, ?, ?, ?)", &AppStmt);
+		if (AppStmt)
+		{
+			sqlite3_bind_text(AppStmt, 1, TCHAR_TO_UTF8(*CachedExeFilename), -1, SQLITE_TRANSIENT);
+			sqlite3_bind_text(AppStmt, 2, TCHAR_TO_UTF8(*CachedAppName), -1, SQLITE_TRANSIENT);
+			sqlite3_bind_text(AppStmt, 3, TCHAR_TO_UTF8(*CachedEngineName), -1, SQLITE_TRANSIENT);
+			sqlite3_bind_int64(AppStmt, 4, (int64)AppVersion);
+			sqlite3_bind_int64(AppStmt, 5, (int64)EngineVersion);
+			const int32 AppIdResult = sqlite3_step(AppStmt);
+			if (AppIdResult != SQLITE_DONE)
+			{
+				UE_LOG(LogASDTool, Warning, TEXT("Failed to insert app_id (sqlite3_step: %d)"), AppIdResult);
+			}
+			sqlite3_finalize(AppStmt);
+		}
+
+		ExecSQL("COMMIT");
+	}
+
+	sqlite3_close(DB);
+	DB = nullptr;
+}
+
+double FDirectSODBFile::GetWriteTimeSeconds() const
+{
+	return FPlatformTime::ToSeconds64(WriteCycles);
+}
+
+
+//--------------------------------------------------------------------------------------------------
+// Helpers
+//--------------------------------------------------------------------------------------------------
+
+bool ReadSODBEntries(const FString& SODBPath, TArray<FSODBEntry>& OutEntries)
+{
+	const FString Filename = FPaths::GetCleanFilename(SODBPath);
+
+	FString KeysPath = SODBPath + TEXT(".keys");
+	if (FPaths::FileExists(KeysPath))
+	{
+		FSODBFactory Factory;
+		if (Factory.Initialize())
+		{
+			FD3D12SODBFile SODBFile;
+			SODBFile.SetFactory(&Factory);
+			if (SODBFile.Open(SODBPath, ESODBOpenMode::Read))
+			{
+				double StartTime = FPlatformTime::Seconds();
+				if (SODBFile.ReadEntries(OutEntries))
+				{
+					UE_LOG(LogASDTool, Display, TEXT("  Read %d entries via D3D12 API in %.3fs: %s"),
+						OutEntries.Num(), FPlatformTime::Seconds() - StartTime, *Filename);
+					SODBFile.Close();
+					return true;
+				}
+				SODBFile.Close();
+			}
+		}
+		UE_LOG(LogASDTool, Warning, TEXT("D3D12 API read failed for '%s', falling back to direct SQLite."), *Filename);
+		OutEntries.Empty();
+	}
+
+	FDirectSODBFile DirectFile;
+	if (!DirectFile.Open(SODBPath, ESODBOpenMode::Read))
+	{
+		UE_LOG(LogASDTool, Error, TEXT("Failed to open SODB '%s' for reading."), *SODBPath);
+		return false;
+	}
+
+	double StartTime = FPlatformTime::Seconds();
+	if (!DirectFile.ReadEntries(OutEntries))
+	{
+		UE_LOG(LogASDTool, Error, TEXT("Failed to read entries from SODB '%s'."), *SODBPath);
+		DirectFile.Close();
+		return false;
+	}
+
+	UE_LOG(LogASDTool, Display, TEXT("  Read %d entries via direct SQLite in %.3fs: %s"),
+		OutEntries.Num(), FPlatformTime::Seconds() - StartTime, *Filename);
+	DirectFile.Close();
+	return true;
+}
+
+//--------------------------------------------------------------------------------------------------
+// PSDB merging
+//--------------------------------------------------------------------------------------------------
+
+/** Get the schema DDL from a SQLite database for comparison. */
+static FString GetPSDBSchemaDDL(const FString& Path)
+{
+	sqlite3* DB = nullptr;
+	sqlite3_open_v2(TCHAR_TO_UTF8(*Path), &DB, SQLITE_OPEN_READONLY, nullptr);
+	if (!DB)
+	{
+		return FString();
+	}
+
+	sqlite3_stmt* Stmt = nullptr;
+	sqlite3_prepare_v2(DB,
+		"SELECT group_concat(sql, ';') FROM sqlite_master WHERE type IN ('table','index') ORDER BY name",
+		-1, &Stmt, nullptr);
+
+	FString DDL;
+	if (sqlite3_step(Stmt) == SQLITE_ROW)
+	{
+		const char* Text = (const char*)sqlite3_column_text(Stmt, 0);
+		if (Text)
+		{
+			DDL = UTF8_TO_TCHAR(Text);
+		}
+	}
+	sqlite3_finalize(Stmt);
+	sqlite3_close(DB);
+	return DDL;
+}
+
+/** Validate that all source PSDBs match the schema of source[0]. */
+static bool ValidatePSDBSchemas(const TArray<FString>& SourcePaths)
+{
+	if (SourcePaths.Num() < 2)
+	{
+		return true;
+	}
+
+	FString RefDDL = GetPSDBSchemaDDL(SourcePaths[0]);
+	if (RefDDL.IsEmpty())
+	{
+		UE_LOG(LogASDTool, Error, TEXT("Schema validation: failed to read schema from '%s'"), *FPaths::GetCleanFilename(SourcePaths[0]));
+		return false;
+	}
+
+	bool bAllMatch = true;
+	for (int32 i = 1; i < SourcePaths.Num(); ++i)
+	{
+		FString DDL = GetPSDBSchemaDDL(SourcePaths[i]);
+		if (DDL.IsEmpty() || DDL != RefDDL)
+		{
+			UE_LOG(LogASDTool, Error, TEXT("Schema mismatch in PSDB '%s'"), *FPaths::GetCleanFilename(SourcePaths[i]));
+			bAllMatch = false;
+		}
+	}
+	return bAllMatch;
+}
+
+bool MergePSDBs(const TArray<FString>& SourcePaths, const FString& DestPath, double& OutMergeSeconds)
+{
+	OutMergeSeconds = 0.0;
+
+	if (SourcePaths.IsEmpty())
+	{
+		UE_LOG(LogASDTool, Error, TEXT("MergePSDBs: no source PSDBs to merge."));
+		return false;
+	}
+
+	if (FPaths::FileExists(DestPath)) { IFileManager::Get().Delete(*DestPath); }
+	IFileManager::Get().MakeDirectory(*FPaths::GetPath(DestPath), true);
+
+	double StartTime = FPlatformTime::Seconds();
+
+	sqlite3* DestDB = nullptr;
+	if (sqlite3_open(TCHAR_TO_UTF8(*DestPath), &DestDB) != SQLITE_OK || !DestDB)
+	{
+		UE_LOG(LogASDTool, Error, TEXT("MergePSDBs: failed to create dest DB '%s'"), *DestPath);
+		return false;
+	}
+
+	auto ExecSQL = [&](const char* SQL) -> bool
+	{
+		char* ErrMsg = nullptr;
+		int32 Rc = sqlite3_exec(DestDB, SQL, nullptr, nullptr, &ErrMsg);
+		if (Rc != SQLITE_OK)
+		{
+			UE_LOG(LogASDTool, Error, TEXT("MergePSDBs SQL error (rc=%d): %s | SQL: %s"),
+				Rc, UTF8_TO_TCHAR(ErrMsg ? ErrMsg : "(null)"), UTF8_TO_TCHAR(SQL));
+			sqlite3_free(ErrMsg);
+			return false;
+		}
+		return true;
+	};
+
+	if (!ValidatePSDBSchemas(SourcePaths))
+	{
+		sqlite3_close(DestDB);
+		IFileManager::Get().Delete(*DestPath);
+		return false;
+	}
+
+	// Copy schema from source[0] directly -- avoids hardcoding, stays in sync with compiler output
+	ExecSQL("PRAGMA synchronous = OFF");
+	{
+		FString EscapedPath = SourcePaths[0].Replace(TEXT("'"), TEXT("''"));
+		if (!ExecSQL(TCHAR_TO_UTF8(*FString::Printf(TEXT("ATTACH DATABASE '%s' AS schema_src"), *EscapedPath))))
+		{
+			sqlite3_close(DestDB);
+			IFileManager::Get().Delete(*DestPath);
+			return false;
+		}
+
+		sqlite3_stmt* SchemaStmt = nullptr;
+		sqlite3_prepare_v2(DestDB,
+			"SELECT sql FROM schema_src.sqlite_master WHERE type IN ('table','index') AND sql IS NOT NULL ORDER BY type DESC, name",
+			-1, &SchemaStmt, nullptr);
+		while (sqlite3_step(SchemaStmt) == SQLITE_ROW)
+		{
+			const char* DDL = (const char*)sqlite3_column_text(SchemaStmt, 0);
+			if (DDL) { ExecSQL(DDL); }
+		}
+		sqlite3_finalize(SchemaStmt);
+
+		// Read user_version and application_id from source[0] via direct connection
+		// (pragma_* virtual tables don't work cross-DB through ATTACH)
+		int64 UserVersion = 0;
+		int64 AppId       = 0;
+		{
+			sqlite3* SrcDB = nullptr;
+			if (sqlite3_open_v2(TCHAR_TO_UTF8(*SourcePaths[0]), &SrcDB, SQLITE_OPEN_READONLY, nullptr) == SQLITE_OK && SrcDB)
+			{
+				auto ReadPragmaInt = [&](const char* SQL) -> int64
+				{
+					sqlite3_stmt* Stmt = nullptr;
+					sqlite3_prepare_v2(SrcDB, SQL, -1, &Stmt, nullptr);
+					int64 Val = 0;
+					if (sqlite3_step(Stmt) == SQLITE_ROW) { Val = sqlite3_column_int64(Stmt, 0); }
+					sqlite3_finalize(Stmt);
+					return Val;
+				};
+				UserVersion = ReadPragmaInt("PRAGMA user_version");
+				AppId       = ReadPragmaInt("PRAGMA application_id");
+				sqlite3_close(SrcDB);
+			}
+		}
+
+		ExecSQL(TCHAR_TO_UTF8(*FString::Printf(TEXT("PRAGMA user_version = %lld"),  UserVersion)));
+		ExecSQL(TCHAR_TO_UTF8(*FString::Printf(TEXT("PRAGMA application_id = %lld"), AppId)));
+		ExecSQL("DETACH DATABASE schema_src");
+	}
+
+	bool bSuccess        = true;
+	bool bMetadataWritten = false;
+
+	for (int32 i = 0; i < SourcePaths.Num() && bSuccess; ++i)
+	{
+		FString EscapedPath = SourcePaths[i].Replace(TEXT("'"), TEXT("''"));
+		if (!ExecSQL(TCHAR_TO_UTF8(*FString::Printf(TEXT("ATTACH DATABASE '%s' AS src"), *EscapedPath))))
+		{
+			bSuccess = false;
+			break;
+		}
+
+		ExecSQL("BEGIN TRANSACTION");
+
+		bSuccess = bSuccess && ExecSQL("INSERT OR IGNORE INTO keys (key, last_accessed) SELECT key, last_accessed FROM src.keys");
+		bSuccess = bSuccess && ExecSQL(
+			"INSERT OR IGNORE INTO vals (keyid, type, value) "
+			"SELECT dk.keyid, sv.type, sv.value "
+			"FROM src.vals sv "
+			"JOIN src.keys sk ON sk.keyid = sv.keyid "
+			"JOIN keys dk ON dk.key = sk.key");
+		bSuccess = bSuccess && ExecSQL("INSERT OR IGNORE INTO groupkeys (key, version) SELECT key, version FROM src.groupkeys");
+		bSuccess = bSuccess && ExecSQL(
+			"INSERT OR IGNORE INTO groups (groupkeyid, valuekeyid) "
+			"SELECT dgk.keyid, dk.keyid "
+			"FROM src.groups sg "
+			"JOIN src.groupkeys sgk ON sgk.keyid = sg.groupkeyid "
+			"JOIN groupkeys dgk ON dgk.key = sgk.key "
+			"JOIN src.keys sk ON sk.keyid = sg.valuekeyid "
+			"JOIN keys dk ON dk.key = sk.key");
+
+		if (!bMetadataWritten)
+		{
+			bSuccess = bSuccess && ExecSQL("INSERT OR IGNORE INTO compiler_id (id, abi, version, family) SELECT id, abi, version, family FROM src.compiler_id");
+			bSuccess = bSuccess && ExecSQL("INSERT OR IGNORE INTO app_id (id, exe_path, app_name, engine_name, app_version, engine_version, app_profile_version) SELECT id, exe_path, app_name, engine_name, app_version, engine_version, app_profile_version FROM src.app_id");
+			bMetadataWritten = true;
+		}
+
+		ExecSQL(bSuccess ? "COMMIT" : "ROLLBACK");
+		ExecSQL("DETACH DATABASE src");
+	}
+
+	if (bSuccess)
+	{
+		auto QueryInt = [&](const char* SQL) -> int32
+		{
+			sqlite3_stmt* Stmt = nullptr;
+			sqlite3_prepare_v2(DestDB, SQL, -1, &Stmt, nullptr);
+			int32 Val = 0;
+			if (sqlite3_step(Stmt) == SQLITE_ROW) { Val = sqlite3_column_int(Stmt, 0); }
+			sqlite3_finalize(Stmt);
+			return Val;
+		};
+		OutMergeSeconds = FPlatformTime::Seconds() - StartTime;
+		UE_LOG(LogASDTool, Display, TEXT("  Merged %d PSDBs -> %d compiled shaders, %d PSO groups in %.2fs: %s"),
+			SourcePaths.Num(),
+			QueryInt("SELECT COUNT(*) FROM keys"),
+			QueryInt("SELECT COUNT(*) FROM groups"),
+			OutMergeSeconds, *FPaths::GetCleanFilename(DestPath));
+	}
+	else
+	{
+		UE_LOG(LogASDTool, Error, TEXT("MergePSDBs failed -- partial data may exist in dest."));
+	}
+
+	sqlite3_close(DestDB);
+	return bSuccess;
+}
+
+}; // namespace ASDTool
\ No newline at end of file
diff --git a/Engine/Source/Programs/Windows/ASDTool/Private/SODBFile.h b/Engine/Source/Programs/Windows/ASDTool/Private/SODBFile.h
new file mode 100644
--- /dev/null
+++ b/Engine/Source/Programs/Windows/ASDTool/Private/SODBFile.h
@@ -0,0 +1,224 @@
+// Copyright Epic Games, Inc. All Rights Reserved.
+
+#pragma once
+
+#include "ASDToolCommands.h"
+#include "ShaderCodeArchive.h"
+
+#include "Windows/AllowWindowsPlatformTypes.h"
+THIRD_PARTY_INCLUDES_START
+#include <d3d12.h>
+THIRD_PARTY_INCLUDES_END
+#include "Windows/HideWindowsPlatformTypes.h"
+
+// Forward declarations for SQLite (avoids IncludeSQLite.h in header)
+struct sqlite3;
+struct sqlite3_stmt;
+
+namespace ASDTool
+{
+
+//--------------------------------------------------------------------------------------------------
+// FSODBEntry  -- A single PSO entry from an SODB, shared between generation and compilation.
+//--------------------------------------------------------------------------------------------------
+
+struct FSODBEntry
+{
+	TArray<uint8> Key;           // Opaque PSO key -- 8 bytes (FShaderHash is FXxHash64, stored big-endian via ToByteArray)
+	uint32 Version = 0;
+	TArray<uint8> StreamData;    // Owns the pipeline state stream memory
+
+	// Whether this is a graphics PSO (vs compute)
+	bool bIsGraphics = false;
+
+	// Per-stage raw DXBC for the direct SQLite read/write path.
+	// Both StreamData and StageDXBC can be populated simultaneously (e.g. for compute, StoreComputeShader
+	// fills StreamData for the D3D12 API write path and StageDXBC[SF_Compute] for the SQLite write path).
+	// Callers should use whichever field is relevant to their write path.
+	TArray<uint8> StageDXBC[SF_NumStandardFrequencies];
+
+	TArray<uint8> RootSigData;             // Serialized root signature bytes (populated by SQLite read)
+	TRefCountPtr<ID3DBlob> RootSigBlob;    // Serialized root signature as ID3DBlob (populated by GenerateSODB)
+
+	/** Get a D3D12_PIPELINE_STATE_STREAM_DESC pointing into StreamData. */
+	D3D12_PIPELINE_STATE_STREAM_DESC GetDesc() const
+	{
+		D3D12_PIPELINE_STATE_STREAM_DESC Desc = {};
+		Desc.SizeInBytes = StreamData.Num();
+		Desc.pPipelineStateSubobjectStream = const_cast<uint8*>(StreamData.GetData());
+		return Desc;
+	}
+};
+
+//--------------------------------------------------------------------------------------------------
+// FSODBGroup
+// One per output .sodb file. Collects FSODBEntry instances during PSO collection phases, then
+// writes them all in a single batch when WriteSODBFiles is called.
+// WantedHashes (from SHK files) defines which compute shaders belong to this group for routing.
+// All public methods are thread-safe.
+//--------------------------------------------------------------------------------------------------
+
+struct FSODBGroup
+{
+	FString Name;
+	TSet<FShaderHash> WantedHashes;
+	TSet<FShaderHash> WrittenHashes;
+	FRWLock WrittenLock;
+	int32 NumFailed = 0;
+	volatile int64 TotalDXILBytes = 0;
+
+	/** Batch-collected entries -- written by WriteSODBFiles. */
+	TArray<FSODBEntry> CollectedEntries;
+	FCriticalSection EntriesLock;
+
+	/** Check if this group wants the hash. No lock needed -- WantedHashes is immutable after SetupGroups. */
+	bool Contains(const FShaderHash& Hash) const;
+
+	/** Build an FSODBEntry from compute shader data and add to CollectedEntries. Thread-safe. */
+	bool StoreComputeShader(const FShaderHash& Hash, const uint8* DXBCData, int64 DXBCSize, ID3DBlob* RootSignatureBlob);
+
+	/** Build an FSODBEntry for a graphics PSO and add to CollectedEntries. Thread-safe.
+	 *  StageDXBC: raw DXBC per stage -- empty arrays for unused stages. VS or MS required (traditional raster or mesh shader pipeline). */
+	bool StoreGraphicsPSO(
+		const FShaderHash& Hash,
+		const TArray<uint8>& DXBC_VS, const TArray<uint8>& DXBC_PS,
+		const TArray<uint8>& DXBC_GS, const TArray<uint8>& DXBC_MS, const TArray<uint8>& DXBC_AS,
+		ID3DBlob* RootSignatureBlob);
+
+private:
+	/** Returns true if this is the first time we've seen this hash. Thread-safe. */
+	bool DeduplicateAndMark(const FShaderHash& Hash);
+};
+
+//--------------------------------------------------------------------------------------------------
+// FSODBFactory  -- Loads D3D12Core.dll once, creates ID3D12StateObjectDatabase instances.
+// Implementation in SODBFile.cpp.
+//--------------------------------------------------------------------------------------------------
+
+class FSODBFactory
+{
+public:
+	~FSODBFactory();
+
+	bool Initialize();
+	ID3D12StateObjectDatabase* CreateDatabase(const FString& OutputPath);
+
+private:
+	HMODULE D3D12Module = nullptr;
+	ID3D12StateObjectDatabaseFactory* Factory = nullptr;
+};
+
+//--------------------------------------------------------------------------------------------------
+// ESODBOpenMode  -- Whether to open an SODB for reading or writing.
+//--------------------------------------------------------------------------------------------------
+
+enum class ESODBOpenMode : uint8
+{
+	Read,
+	Write
+};
+
+//--------------------------------------------------------------------------------------------------
+// ISODBFile  -- Abstract interface for reading/writing SODB files.
+//--------------------------------------------------------------------------------------------------
+
+class ISODBFile
+{
+public:
+	virtual ~ISODBFile() = default;
+
+	/** Open an SODB file for reading or writing. */
+	virtual bool Open(const FString& Path, ESODBOpenMode Mode) = 0;
+
+	/** Set application metadata. */
+	virtual void SetApplicationDesc(const D3D12_APPLICATION_DESC& AppDesc) = 0;
+
+	/** Write all entries to the SODB. */
+	virtual bool WriteEntries(const TArray<FSODBEntry>& Entries) = 0;
+
+	/** Read all entries from the SODB. */
+	virtual bool ReadEntries(TArray<FSODBEntry>& OutEntries) = 0;
+
+	/** Close/finalize the database. */
+	virtual void Close() = 0;
+
+	/** Get timing stats in seconds. */
+	virtual double GetWriteTimeSeconds() const = 0;
+};
+
+//--------------------------------------------------------------------------------------------------
+// FD3D12SODBFile  -- Writes/reads via ID3D12StateObjectDatabase API (official).
+// Writes a key manifest file (.keys) alongside the SODB for enumeration.
+//--------------------------------------------------------------------------------------------------
+
+class FD3D12SODBFile : public ISODBFile
+{
+public:
+	~FD3D12SODBFile();
+
+	/** Set the factory to use. Must be called before Open(). */
+	void SetFactory(FSODBFactory* InFactory) { Factory = InFactory; }
+
+	virtual bool Open(const FString& Path, ESODBOpenMode Mode) override;
+	virtual void SetApplicationDesc(const D3D12_APPLICATION_DESC& AppDesc) override;
+	virtual bool WriteEntries(const TArray<FSODBEntry>& Entries) override;
+	virtual bool ReadEntries(TArray<FSODBEntry>& OutEntries) override;
+	virtual void Close() override;
+	virtual double GetWriteTimeSeconds() const override;
+
+private:
+	FSODBFactory* Factory = nullptr;
+	ID3D12StateObjectDatabase* Database = nullptr;
+	FString FilePath;
+	volatile int64 WriteCycles = 0;
+};
+
+//--------------------------------------------------------------------------------------------------
+// FDirectSODBFile  -- Writes/reads SODB as SQLite database directly (fast, single transaction).
+//--------------------------------------------------------------------------------------------------
+
+class FDirectSODBFile : public ISODBFile
+{
+public:
+	~FDirectSODBFile();
+
+	virtual bool Open(const FString& Path, ESODBOpenMode Mode) override;
+	virtual void SetApplicationDesc(const D3D12_APPLICATION_DESC& AppDesc) override;
+	virtual bool WriteEntries(const TArray<FSODBEntry>& Entries) override;
+	virtual bool ReadEntries(TArray<FSODBEntry>& OutEntries) override;
+	virtual void Close() override;
+	virtual double GetWriteTimeSeconds() const override;
+
+private:
+	sqlite3* DB = nullptr;
+	sqlite3_stmt* StmtRootSig = nullptr;
+	sqlite3_stmt* StmtBytecode = nullptr;
+	sqlite3_stmt* StmtPSO = nullptr;
+	sqlite3_stmt* StmtPSO_Graphics = nullptr;
+	sqlite3_stmt* StmtGroup = nullptr;
+	volatile int64 WriteCycles = 0;
+	ESODBOpenMode OpenMode = ESODBOpenMode::Write;
+	FString FilePath;
+	D3D12_APPLICATION_DESC CachedAppDesc = {};
+	FString CachedExeFilename;
+	FString CachedAppName;
+	FString CachedEngineName;
+
+	void Cleanup();
+	void ExecSQL(const char* SQL);
+	void PrepareSQLite(const char* SQL, sqlite3_stmt** OutStmt);
+	void CreateSchema();
+};
+
+//--------------------------------------------------------------------------------------------------
+// Helpers
+//--------------------------------------------------------------------------------------------------
+
+/** Read PSO entries from an SODB. Tries D3D12 API path first (.keys manifest), falls back to direct SQLite. */
+bool ReadSODBEntries(const FString& SODBPath, TArray<FSODBEntry>& OutEntries);
+
+/** Merge N per-thread PSDBs into a single destination PSDB via SQLite ATTACH.
+ *  All inserts are content-addressed (INSERT OR IGNORE), keyids remapped through key content. */
+bool MergePSDBs(const TArray<FString>& SourcePaths, const FString& DestPath, double& OutMergeSeconds);
+
+};
\ No newline at end of file
diff --git a/Engine/Source/Runtime/D3D12RHI/Internal/D3D12ShaderUtils.h b/Engine/Source/Runtime/D3D12RHI/Internal/D3D12ShaderUtils.h
--- a/Engine/Source/Runtime/D3D12RHI/Internal/D3D12ShaderUtils.h
+++ b/Engine/Source/Runtime/D3D12RHI/Internal/D3D12ShaderUtils.h
@@ -20,3 +20,19 @@ D3D12RHI_API bool BuildComputeShaderRootSignature(
 	D3D12_RESOURCE_BINDING_TIER ResourceBindingTier,
 	TRefCountPtr<ID3DBlob>& OutRootSigBlob,
 	int32& OutNativeBytecodeOffset);
+
+/**
+ * Build a serialized root signature for a graphics PSO from multiple UE shader blobs.
+ * Merges FShaderCodePackedResourceCounts across all active shader stages to produce
+ * a single root signature that covers the union of resources required by all stages.
+ *
+ * @param ShaderBlobs           Array of full UE shader blobs, one per active stage (SRT + native bytecode + optional data).
+ *                              Null/empty entries are skipped.
+ * @param ResourceBindingTier   Resource binding tier to use for the root signature layout
+ * @param OutRootSigBlob        Receives the serialized root signature blob
+ * @return true on success
+ */
+D3D12RHI_API bool BuildGraphicsShaderRootSignature(
+	const TConstArrayView<TConstArrayView<uint8>> ShaderBlobs,
+	D3D12_RESOURCE_BINDING_TIER ResourceBindingTier,
+	TRefCountPtr<ID3DBlob>& OutRootSigBlob);
diff --git a/Engine/Source/Runtime/D3D12RHI/Private/D3D12Util.cpp b/Engine/Source/Runtime/D3D12RHI/Private/D3D12Util.cpp
--- a/Engine/Source/Runtime/D3D12RHI/Private/D3D12Util.cpp
+++ b/Engine/Source/Runtime/D3D12RHI/Private/D3D12Util.cpp
@@ -2010,4 +2010,87 @@ bool BuildComputeShaderRootSignature(const TConstArrayView<uint8>& ShaderCode, D
 	return true;
 }
 
+//--------------------------------------------------------------------------------------------------
+// BuildGraphicsShaderRootSignature -- Build a serialized root signature from multiple UE shader blobs.
+// Merges resource counts across all active stages (VS, PS, GS, MS, AS).
+// Can be called without an initialized D3D12 RHI.
+//--------------------------------------------------------------------------------------------------
+
+bool BuildGraphicsShaderRootSignature(const TConstArrayView<TConstArrayView<uint8>> ShaderBlobs,
+	D3D12_RESOURCE_BINDING_TIER ResourceBindingTier, TRefCountPtr<ID3DBlob>& OutRootSigBlob)
+{
+	if (ShaderBlobs.IsEmpty())
+	{
+		return false;
+	}
+
+	// Merge FShaderCodePackedResourceCounts across all active shader stages.
+	// The root signature must cover the union of all resources required by any stage.
+	FShaderCodePackedResourceCounts MergedCounts = {};
+	bool bFoundAny = false;
+
+	for (const TConstArrayView<uint8>& ShaderCode : ShaderBlobs)
+	{
+		if (ShaderCode.IsEmpty())
+		{
+			continue;
+		}
+
+		FShaderCodeReader CodeReader(ShaderCode);
+
+		const FShaderCodePackedResourceCounts* PackedCounts = CodeReader.FindOptionalData<FShaderCodePackedResourceCounts>();
+		if (!PackedCounts)
+		{
+			// Not fatal -- some stages may not have resource counts (e.g. pass-through GS)
+			continue;
+		}
+
+		// Merge by taking the max of each resource count across stages
+		MergedCounts.NumSamplers = FMath::Max(MergedCounts.NumSamplers, PackedCounts->NumSamplers);
+		MergedCounts.NumSRVs     = FMath::Max(MergedCounts.NumSRVs,     PackedCounts->NumSRVs);
+		MergedCounts.NumCBs      = FMath::Max(MergedCounts.NumCBs,      PackedCounts->NumCBs);
+		MergedCounts.NumUAVs     = FMath::Max(MergedCounts.NumUAVs,     PackedCounts->NumUAVs);
+		MergedCounts.UsageFlags |= PackedCounts->UsageFlags;
+		bFoundAny = true;
+	}
+
+	if (!bFoundAny)
+	{
+		UE_LOG(LogD3D12RHI, Warning, TEXT("BuildGraphicsShaderRootSignature: No resource counts found in any shader stage."));
+		return false;
+	}
+
+	// Build quantized bound shader state from merged resource counts -- same path as compute
+	GRHISupportsShaderRootConstants = true;
+
+	FD3D12ShaderData TempShaderData;
+	TempShaderData.ResourceCounts = MergedCounts;
+	// NOTE: VendorExtensions (AGS/NV intrinsics) are not yet extracted from graphics shader stages
+	// and are therefore not forwarded to QuantizeBoundShaderStateCommon. Graphics PSOs that rely on
+	// vendor intrinsics will get a root signature without the extra descriptor space those require.
+	// This is acceptable for the current ASD tool scope -- support can be added when needed.
+
+	FD3D12QuantizedBoundShaderState QBSS{};
+	// Use SV_All visibility so the root signature covers all stages
+	QuantizeBoundShaderStateCommon(QBSS, &TempShaderData, ResourceBindingTier, SV_All, true /*bAllowUAVs*/);
+
+	FD3D12RootSignatureDesc RootSigDesc(QBSS, ResourceBindingTier);
+
+	TRefCountPtr<ID3DBlob> ErrorBlob;
+	HRESULT hr = D3DX12SerializeVersionedRootSignature(
+		&RootSigDesc.GetDesc(),
+		D3D_ROOT_SIGNATURE_VERSION_1_1,
+		OutRootSigBlob.GetInitReference(),
+		ErrorBlob.GetInitReference());
+
+	if (FAILED(hr))
+	{
+		UE_LOG(LogD3D12RHI, Error, TEXT("BuildGraphicsShaderRootSignature: Failed to serialize root signature: 0x%08X %s"),
+			hr, ErrorBlob ? ANSI_TO_TCHAR((const char*)ErrorBlob->GetBufferPointer()) : TEXT(""));
+		return false;
+	}
+
+	return true;
+}
+
 #undef LOCTEXT_NAMESPACE
