ASD tool fixes:
- delete keys file on clear sodb data (sidecar data which needs to be removed as well)
- make sure 16 byte keys are used in SqlLite SODB database because d3d12 compiler expects this
- rename all the print functions to make it clear
- disable d3d12 SODB reading code path because the extracted data can't be used for compilation according to the doc so fallback to SQLLite reading for now

#Synced-CL 52245732

diff --git a/Engine/Source/Programs/Windows/ASDTool/Private/GenerateSODB.cpp b/Engine/Source/Programs/Windows/ASDTool/Private/GenerateSODB.cpp
--- a/Engine/Source/Programs/Windows/ASDTool/Private/GenerateSODB.cpp
+++ b/Engine/Source/Programs/Windows/ASDTool/Private/GenerateSODB.cpp
@@ -88,7 +88,7 @@ struct FGenerateSODBStats
 // Usage / arg parsing
 //--------------------------------------------------------------------------------------------------
 
-static void PrintUsage()
+static void PrintGenerateSODBUsage()
 {
 	UE_LOG(LogASDTool, Display, TEXT("GenerateSODB: generate State Object Database (.sodb) files for compute and graphics shaders"));
 	UE_LOG(LogASDTool, Display, TEXT("  Extracts shader DXBC bytecode from cooked shader archives and writes"));
@@ -141,7 +141,7 @@ static bool ParseGenerateSODBArgs(FGenerateSODBArgs& OutArgs)
 	if (OutArgs.OutputDir.IsEmpty() || OutArgs.ShaderCodeDir.IsEmpty() || !bHasCollectionSource ||
 		OutArgs.AppName.IsEmpty() || OutArgs.ExeFilename.IsEmpty() || OutArgs.EngineName.IsEmpty())
 	{
-		PrintUsage();
+		PrintGenerateSODBUsage();
 		if (OutArgs.OutputDir.IsEmpty())        { UE_LOG(LogASDTool, Error, TEXT("No output directory specified.")); }
 		if (OutArgs.ShaderCodeDir.IsEmpty())    { UE_LOG(LogASDTool, Error, TEXT("No shader code directory specified.")); }
 		if (!bHasCollectionSource)              { UE_LOG(LogASDTool, Error, TEXT("Specify at least one of: -AllCompute, -BundleDir (-ShkDir alone only sets up group routing)")); }
@@ -162,7 +162,8 @@ static bool ParseGenerateSODBArgs(FGenerateSODBArgs& OutArgs)
 
 //--------------------------------------------------------------------------------------------------
 // PrepareOutputDirectory
-// Deletes stale .sodb files from the output directory if -Clean was requested, then creates it.
+// Deletes stale .sodb and .sodb.keys files from the output directory if -Clean was requested,
+// then creates the directory.
 //--------------------------------------------------------------------------------------------------
 
 static int32 PrepareOutputDirectory(const FGenerateSODBArgs& Args)
@@ -173,15 +174,27 @@ static int32 PrepareOutputDirectory(const FGenerateSODBArgs& Args)
 		IFileManager::Get().FindFiles(ExistingSODBs, *Args.OutputDir, TEXT(".sodb"));
 		for (const FString& Existing : ExistingSODBs)
 		{
-			const FString FullPath = Args.OutputDir / Existing;
-			if (!IFileManager::Get().Delete(*FullPath))
+			const FString SODBPath = Args.OutputDir / Existing;
+			if (!IFileManager::Get().Delete(*SODBPath))
 			{
-				UE_LOG(LogASDTool, Warning, TEXT("Failed to delete: %s"), *FullPath);
+				UE_LOG(LogASDTool, Warning, TEXT("Failed to delete: %s"), *SODBPath);
 			}
 		}
-		if (ExistingSODBs.Num() > 0)
+
+		TArray<FString> ExistingKeys;
+		IFileManager::Get().FindFiles(ExistingKeys, *Args.OutputDir, TEXT(".sodb.keys"));
+		for (const FString& Existing : ExistingKeys)
+		{
+			const FString KeysPath = Args.OutputDir / Existing;
+			if (!IFileManager::Get().Delete(*KeysPath))
+			{
+				UE_LOG(LogASDTool, Warning, TEXT("Failed to delete: %s"), *KeysPath);
+			}
+		}
+
+		if (ExistingSODBs.Num() > 0 || ExistingKeys.Num() > 0)
 		{
-			UE_LOG(LogASDTool, Display, TEXT("Cleaned %d existing .sodb files."), ExistingSODBs.Num());
+			UE_LOG(LogASDTool, Display, TEXT("Cleaned %d existing .sodb file(s) (and %d .sodb.keys files)."), ExistingSODBs.Num(), ExistingKeys.Num());
 		}
 	}
 
@@ -901,7 +914,7 @@ int32 GenerateSODBs()
 
 	if (FParse::Param(CmdLine, TEXT("h")) || FParse::Param(CmdLine, TEXT("help")))
 	{
-		PrintUsage();
+		PrintGenerateSODBUsage();
 		return 0;
 	}
 
diff --git a/Engine/Source/Programs/Windows/ASDTool/Private/RegisterPSDBs.cpp b/Engine/Source/Programs/Windows/ASDTool/Private/RegisterPSDBs.cpp
--- a/Engine/Source/Programs/Windows/ASDTool/Private/RegisterPSDBs.cpp
+++ b/Engine/Source/Programs/Windows/ASDTool/Private/RegisterPSDBs.cpp
@@ -487,7 +487,7 @@ struct FRegisterPSDBsArgs
 // Usage
 // ------------------------------------------------------------------
 
-static void PrintUsage()
+static void PrintRegisterPSDBsUsage()
 {
 	UE_LOG(LogASDTool, Display, TEXT("RegisterPSDBs: register Precompiled Shader Databases with the D3D12 runtime"));
 	UE_LOG(LogASDTool, Display, TEXT(""));
@@ -546,7 +546,7 @@ static bool ValidateRegistrationArgs(const FRegisterPSDBsArgs& Args)
 	if (Args.PSDBDir.IsEmpty() || Args.SODBPath.IsEmpty() || Args.AppName.IsEmpty() ||
 		Args.ExeFilename.IsEmpty() || Args.ExePath.IsEmpty() || Args.EngineName.IsEmpty())
 	{
-		PrintUsage();
+		PrintRegisterPSDBsUsage();
 		if (Args.PSDBDir.IsEmpty())     UE_LOG(LogASDTool, Error, TEXT("Missing -PSDBDir"));
 		if (Args.SODBPath.IsEmpty())    UE_LOG(LogASDTool, Error, TEXT("Missing -SODBPath"));
 		if (Args.AppName.IsEmpty())     UE_LOG(LogASDTool, Error, TEXT("Missing -Name"));
@@ -687,7 +687,7 @@ int32 RegisterPSDBs()
 
 	if (FParse::Param(CmdLine, TEXT("h")) || FParse::Param(CmdLine, TEXT("help")))
 	{
-		PrintUsage();
+		PrintRegisterPSDBsUsage();
 		return 0;
 	}
 
diff --git a/Engine/Source/Programs/Windows/ASDTool/Private/SODBFile.cpp b/Engine/Source/Programs/Windows/ASDTool/Private/SODBFile.cpp
--- a/Engine/Source/Programs/Windows/ASDTool/Private/SODBFile.cpp
+++ b/Engine/Source/Programs/Windows/ASDTool/Private/SODBFile.cpp
@@ -740,9 +740,13 @@ bool FDirectSODBFile::WriteEntries(const TArray<FSODBEntry>& Entries)
 			}
 		}
 
-		// Use the full 8-byte FXxHash64 key for both PSO and group entries
-		const uint8* PSOKey = Entry.Key.GetData();
-		const int    PSOKeySize = Entry.Key.Num();
+		// PSO key: 16 bytes, zero-padded from Entry.Key (8-byte FXxHash64).
+		// The D3D12 compiler exe expects 16-byte keys in pipeline_states.Key and groups.PSOKey.
+		// We use Entry.Key directly (padded) rather than re-hashing -- it already uniquely
+		// identifies the PSO for all shader types (compute and graphics). Padding to 16 bytes
+		// does not reduce collision resistance since Entry.Key is already a 64-bit hash.
+		uint8 FinalPSOKey[16] = {};
+		FMemory::Memcpy(FinalPSOKey, Entry.Key.GetData(), FMath::Min(Entry.Key.Num(), 16));
 		const void* RootSigKeyPtr = Entry.RootSigBlob ? RootSigKey : nullptr;
 		const int RootSigKeySize = Entry.RootSigBlob ? 16 : 0;
 
@@ -759,7 +763,7 @@ bool FDirectSODBFile::WriteEntries(const TArray<FSODBEntry>& Entries)
 			bool bHasAS = InsertShaderBytecode(StmtBytecode, Entry.StageDXBC[SF_Amplification], "AS", KeyAS);
 
 			// Insert graphics pipeline state
-			sqlite3_bind_blob(StmtPSO_Graphics, 1, PSOKey, PSOKeySize, SQLITE_STATIC);
+			sqlite3_bind_blob(StmtPSO_Graphics, 1, FinalPSOKey, 16, SQLITE_STATIC);
 			sqlite3_bind_blob(StmtPSO_Graphics, 2, RootSigKeyPtr, RootSigKeySize, SQLITE_STATIC);
 			sqlite3_bind_blob(StmtPSO_Graphics, 3, bHasVS ? KeyVS : nullptr, bHasVS ? 16 : 0, SQLITE_STATIC);
 			sqlite3_bind_blob(StmtPSO_Graphics, 4, bHasPS ? KeyPS : nullptr, bHasPS ? 16 : 0, SQLITE_STATIC);
@@ -783,8 +787,7 @@ bool FDirectSODBFile::WriteEntries(const TArray<FSODBEntry>& Entries)
 			{
 				UE_LOG(LogASDTool, Warning, TEXT("Failed to insert shader bytecode"));
 			}
-
-			sqlite3_bind_blob(StmtPSO, 1, PSOKey, PSOKeySize, SQLITE_STATIC);
+			sqlite3_bind_blob(StmtPSO, 1, FinalPSOKey, 16, SQLITE_STATIC);
 			sqlite3_bind_blob(StmtPSO, 2, RootSigKeyPtr, RootSigKeySize, SQLITE_STATIC);
 			sqlite3_bind_blob(StmtPSO, 3, BytecodeKey, 16, SQLITE_STATIC);
 			StepResult = sqlite3_step(StmtPSO);
@@ -796,9 +799,9 @@ bool FDirectSODBFile::WriteEntries(const TArray<FSODBEntry>& Entries)
 			UE_LOG(LogASDTool, Warning, TEXT("Failed to insert pipeline state"));
 		}
 
-		// Insert group entry -- key and PSO key are both the full 8-byte FXxHash64
+		// Insert group entry -- Key is 8-byte FXxHash64 (routing), PSOKey FK is 16-byte
 		sqlite3_bind_blob(StmtGroup, 1, Entry.Key.GetData(), Entry.Key.Num(), SQLITE_STATIC);
-		sqlite3_bind_blob(StmtGroup, 2, PSOKey, PSOKeySize, SQLITE_STATIC);
+		sqlite3_bind_blob(StmtGroup, 2, FinalPSOKey, 16, SQLITE_STATIC);
 		StepResult = sqlite3_step(StmtGroup);
 		sqlite3_reset(StmtGroup);
 		if (StepResult != SQLITE_DONE)
@@ -935,8 +938,13 @@ bool ReadSODBEntries(const FString& SODBPath, TArray<FSODBEntry>& OutEntries)
 {
 	const FString Filename = FPaths::GetCleanFilename(SODBPath);
 
+	// Disable the d3d12 sodb read path because the darta can't be used for PSO compilation
+	// becasue the retrieved data can't be used directly as noted in the spec:
+	// Note: May contain CD3DX12_GLOBAL_SERIALIZED_ROOT_SIGNATURE_SUBOBJECT or CD3DX12_LOCAL_SERIALIZED_ROOT_SIGNATURE_SUBOBJECT
+	// which cannot currently be consumed by the compiler or runtime.
+	bool bAllowD3D12ReadPath = false;
 	FString KeysPath = SODBPath + TEXT(".keys");
-	if (FPaths::FileExists(KeysPath))
+	if (bAllowD3D12ReadPath && FPaths::FileExists(KeysPath))
 	{
 		FSODBFactory Factory;
 		if (Factory.Initialize())
